Copying //UE4/Dev-Sequencer to //UE4/Main (Source: //UE4/Dev-Sequencer @ 3617642)

diff --git a/Engine/Source/Runtime/MediaAssets/Public/MediaSoundComponent.h b/Engine/Source/Runtime/MediaAssets/Public/MediaSoundComponent.h
new file mode 100644
--- /dev/null
+++ b/Engine/Source/Runtime/MediaAssets/Public/MediaSoundComponent.h
@@ -0,0 +1,109 @@
+// Copyright 1998-2017 Epic Games, Inc. All Rights Reserved.
+
+#pragma once
+
+#include "Components/SynthComponent.h"
+
+#include "Containers/Array.h"
+#include "HAL/CriticalSection.h"
+#include "MediaSampleQueue.h"
+#include "Misc/Timespan.h"
+#include "Templates/SharedPointer.h"
+#include "UObject/ObjectMacros.h"
+#include "UObject/ScriptMacros.h"
+
+#include "MediaSoundComponent.generated.h"
+
+class FMediaAudioResampler;
+class FMediaPlayerFacade;
+class IMediaAudioSample;
+class IMediaPlayer;
+class UMediaPlayer;
+
+
+/**
+ * Available media sound channel types.
+ */
+UENUM()
+enum class EMediaSoundChannels
+{
+	/** Mono (1 channel). */
+	Mono,
+
+	/** Stereo (2 channels). */
+	Stereo,
+
+	/** Surround sound (7.1 channels; for UI). */
+	Surround
+};
+
+
+/**
+ * Implements a sound component for playing a media player's audio output.
+ */
+UCLASS(ClassGroup=Media, editinlinenew, meta=(BlueprintSpawnableComponent))
+class MEDIAASSETS_API UMediaSoundComponent
+	: public USynthComponent
+{
+	GENERATED_BODY()
+
+public:
+
+	/** Media sound channel type. */
+	UPROPERTY(EditAnywhere, Category="UMediaSoundComponent")
+	EMediaSoundChannels Channels;
+
+	/** The media player asset associated with this component. */
+	UPROPERTY(BlueprintReadWrite, EditAnywhere, Category="Media")
+	UMediaPlayer* MediaPlayer;
+
+public:
+
+	/**
+	 * Create and initialize a new instance.
+	 *
+	 * @param ObjectInitializer Initialization parameters.
+	 */
+	UMediaSoundComponent(const FObjectInitializer& ObjectInitializer);
+
+	/** Virtual destructor. */
+	~UMediaSoundComponent();
+
+public:
+
+	void UpdatePlayer();
+
+public:
+
+	//~ UActorComponent interface
+
+	virtual void TickComponent(float DeltaTime, enum ELevelTick TickType, FActorComponentTickFunction *ThisTickFunction) override;
+
+public:
+
+	//~ USceneComponent interface
+
+	virtual void Activate(bool bReset = false) override;
+	virtual void Deactivate() override;
+
+protected:
+
+	//~ USynthComponent interface
+
+	virtual void Init(const int32 SampleRate) override;
+	virtual void OnGenerateAudio(TArray<float>& OutAudio) override;
+
+private:
+
+	/** Critical section for synchronizing access to PlayerFacadePtr. */
+	FCriticalSection CriticalSection;
+
+	/** The player facade that's currently providing texture samples. */
+	TWeakPtr<FMediaPlayerFacade, ESPMode::ThreadSafe> CurrentPlayerFacade;
+
+	/** The audio resampler. */
+	FMediaAudioResampler* Resampler;
+
+	/** Audio sample queue. */
+	TSharedPtr<FMediaAudioSampleQueue, ESPMode::ThreadSafe> SampleQueue;
+};
diff --git a/Engine/Source/Runtime/MediaAssets/Public/MediaSource.h b/Engine/Source/Runtime/MediaAssets/Public/MediaSource.h
--- a/Engine/Source/Runtime/MediaAssets/Public/MediaSource.h
+++ b/Engine/Source/Runtime/MediaAssets/Public/MediaSource.h
@@ -2,10 +2,12 @@
 
 #pragma once
 
-#include "CoreMinimal.h"
+#include "CoreTypes.h"
+#include "IMediaOptions.h"
 #include "UObject/ObjectMacros.h"
 #include "UObject/Object.h"
 #include "IMediaOptions.h"
+#include "UObject/ScriptMacros.h"
 
 #include "MediaSource.generated.h"
 
@@ -15,7 +17,9 @@
  *
  * Media sources describe the location and/or settings of media objects that can
  * be played in a media player, such as a video file on disk, a video stream on
- * the internet, or a web cam attached to or built into the target device.
+ * the internet, or a web cam attached to or built into the target device. The
+ * location is encoded as a media URL string, whose URI scheme and optional file
+ * extension will be used to locate a suitable media player.
  */
 UCLASS(Abstract, BlueprintType, hidecategories=(Object))
 class MEDIAASSETS_API UMediaSource
@@ -26,10 +30,21 @@ class MEDIAASSETS_API UMediaSource
 
 public:
 
-	/** Get the media source's URL string (must override in child classes). */
+	/**
+	 * Get the media source's URL string (must be implemented in child classes).
+	 *
+	 * @return The media URL.
+	 * @see GetProxies
+	 */
+	UFUNCTION(BlueprintCallable, Category="Media|MediaSource")
 	virtual FString GetUrl() const PURE_VIRTUAL(UMediaSource::GetUrl, return FString(););
 
-	/** Validate the media source settings (must override in child classes). */
+	/**
+	 * Validate the media source settings (must be implemented in child classes).
+	 *
+	 * @return true if validation passed, false otherwise.
+	 */
+	UFUNCTION(BlueprintCallable, Category="Media|MediaSource")
 	virtual bool Validate() const PURE_VIRTUAL(UMediaSource::Validate, return false;);
 
 public:
diff --git a/Engine/Source/Runtime/MediaAssets/Public/MediaTexture.h b/Engine/Source/Runtime/MediaAssets/Public/MediaTexture.h
--- a/Engine/Source/Runtime/MediaAssets/Public/MediaTexture.h
+++ b/Engine/Source/Runtime/MediaAssets/Public/MediaTexture.h
@@ -2,37 +2,54 @@
 
 #pragma once
 
-#include "CoreMinimal.h"
+#include "CoreTypes.h"
+#include "Containers/EnumAsByte.h"
+#include "Engine/Texture.h"
+#include "Math/Color.h"
+#include "Math/IntPoint.h"
+#include "MediaSampleQueue.h"
+#include "Misc/Timespan.h"
+#include "Templates/SharedPointer.h"
 #include "UObject/ObjectMacros.h"
 #include "UObject/ScriptMacros.h"
-#include "Engine/Texture.h"
-#include "IMediaTextureSink.h"
+
 #include "MediaTexture.generated.h"
 
-class FTextureResource;
+class FMediaPlayerFacade;
+class FMediaTextureClockSink;
+class IMediaTextureSample;
+class UMediaPlayer;
+
 
 /**
  * Implements a texture asset for rendering video tracks from UMediaPlayer assets.
  */
-UCLASS(hidecategories=(Compression, LevelOfDetail, Object, Texture))
+UCLASS(hidecategories=(Adjustments, Compositing, LevelOfDetail, Object))
 class MEDIAASSETS_API UMediaTexture
 	: public UTexture
-	, public IMediaTextureSink
 {
 	GENERATED_UCLASS_BODY()
 
 	/** The addressing mode to use for the X axis. */
-	UPROPERTY(EditAnywhere, BlueprintReadWrite, Category="Media|MediaTexture", AssetRegistrySearchable)
+	UPROPERTY(EditAnywhere, BlueprintReadWrite, Category="MediaTexture", meta=(DisplayName="X-axis Tiling Method"), AssetRegistrySearchable, AdvancedDisplay)
 	TEnumAsByte<TextureAddress> AddressX;
 
 	/** The addressing mode to use for the Y axis. */
-	UPROPERTY(EditAnywhere, BlueprintReadWrite, Category="Media|MediaTexture", AssetRegistrySearchable)
+	UPROPERTY(EditAnywhere, BlueprintReadWrite, Category="MediaTexture", meta=(DisplayName="Y-axis Tiling Method"), AssetRegistrySearchable, AdvancedDisplay)
 	TEnumAsByte<TextureAddress> AddressY;
 
-	/** The color used to clear the texture if CloseAction is set to Clear (default = black). */
-	UPROPERTY(EditAnywhere, BlueprintReadWrite, Category="Media|MediaTexture", AdvancedDisplay)
+	/** Whether to clear the texture when no media is being played (default = enabled). */
+	UPROPERTY(EditAnywhere, BlueprintReadWrite, Category="MediaTexture")
+	bool AutoClear;
+
+	/** The color used to clear the texture if AutoClear is enabled (default = black). */
+	UPROPERTY(EditAnywhere, BlueprintReadWrite, Category="MediaTexture")
 	FLinearColor ClearColor;
 
+	/** The media player asset associated with this texture. */
+	UPROPERTY(EditAnywhere, BlueprintReadWrite, Category="Media")
+	UMediaPlayer* MediaPlayer;
+
 public:
 
 	/**
@@ -62,88 +79,59 @@ public:
 	UFUNCTION(BlueprintCallable, Category="Media|MediaTexture")
 	int32 GetWidth() const;
 
-public:
-
-	/**
-	 * Get the event delegate that is invoked when this asset is being destroyed.
-	 *
-	 * @return The delegate.
-	 */
-	DECLARE_EVENT_OneParam(UMediaTexture, FOnBeginDestroy, UMediaTexture& /*DestroyedMediaTexture*/)
-	FOnBeginDestroy& OnBeginDestroy()
-	{
-		return BeginDestroyEvent;
-	}
-
 public:
 
 	//~ UTexture interface.
 
+	virtual void BeginDestroy() override;
 	virtual FTextureResource* CreateResource() override;
-	virtual EMaterialValueType GetMaterialType() override;
+	virtual EMaterialValueType GetMaterialType() const override;
 	virtual float GetSurfaceWidth() const override;
 	virtual float GetSurfaceHeight() const override;
-	virtual void UpdateResource() override;
+	virtual FGuid GetExternalTextureGuid() const override;
 
 public:
 
 	//~ UObject interface.
 
-	virtual void BeginDestroy() override;
 	virtual FString GetDesc() override;
 	virtual void GetResourceSizeEx(FResourceSizeEx& CumulativeResourceSize) override;
 
 #if WITH_EDITOR
 	virtual void PostEditChangeProperty(FPropertyChangedEvent& PropertyChangedEvent) override;
-	virtual void PreEditChange(UProperty* PropertyAboutToChange) override;
 #endif
 
-public:
-
-	//~ IMediaTextureSink interface
-
-	virtual void* AcquireTextureSinkBuffer() override;
-	virtual void DisplayTextureSinkBuffer(FTimespan Time) override;
-	virtual FIntPoint GetTextureSinkDimensions() const override;
-	virtual EMediaTextureSinkFormat GetTextureSinkFormat() const override;
-	virtual EMediaTextureSinkMode GetTextureSinkMode() const override;
-	virtual FRHITexture* GetTextureSinkTexture() override;
-	virtual bool InitializeTextureSink(FIntPoint OutputDim, FIntPoint BufferDim, EMediaTextureSinkFormat Format, EMediaTextureSinkMode Mode) override;
-	virtual void ReleaseTextureSinkBuffer() override;
-	virtual void ShutdownTextureSink() override;
-	virtual bool SupportsTextureSinkFormat(EMediaTextureSinkFormat Format) const override;
-	virtual void UpdateTextureSinkBuffer(const uint8* Data, uint32 Pitch = 0) override;
-	virtual void UpdateTextureSinkResource(FRHITexture* RenderTarget, FRHITexture* ShaderResource) override;
-
 protected:
 
-	//~ Deprecated members
+	/**
+	 * Tick the texture resource.
+	 *
+	 * @param Timecode The current timecode.
+	 */
+	void TickResource(FTimespan Timecode);
 
-	DEPRECATED_FORGAME(4.13, "The MediaPlayer property is no longer used. Please upgrade your content to Media Framework 2.0")
-	UPROPERTY(BlueprintReadWrite, Category=MediaPlayer)
-	class UMediaPlayer* MediaPlayer;
-	
-	DEPRECATED_FORGAME(4.13, "The VideoTrackIndex property is no longer used. Please upgrade your content to Media Framework 2.0")
-	UPROPERTY(BlueprintReadWrite, Category=MediaPlayer)
-	int32 VideoTrackIndex;
+	/** Unregister the player's external texture GUID. */
+	void UnregisterPlayerGuid();
 
 private:
 
-	/** An event delegate that is invoked when this asset is being destroyed. */
-	FOnBeginDestroy BeginDestroyEvent;
+	friend class FMediaTextureClockSink;
+
+	/** The texture's media clock sink. */
+	TSharedPtr<FMediaTextureClockSink, ESPMode::ThreadSafe> ClockSink;
 
-	/** Critical section for synchronizing access to texture resource object. */
-	mutable FCriticalSection CriticalSection;
+	/** The player facade that's currently providing texture samples. */
+	TWeakPtr<FMediaPlayerFacade, ESPMode::ThreadSafe> CurrentPlayerFacade;
 
-	/** Width and height of the sink buffers (in pixels). */
-	FIntPoint SinkBufferDim;
+	/** Current width and height of the resource (in pixels). */
+	FIntPoint Dimensions;
 
-	/** Width and height of the video output (in pixels). */
-	FIntPoint SinkOutputDim;
+	/** The previously used player GUID. */
+	FGuid LastPlayerGuid;
 
-	/** The render target's pixel format. */
-	EMediaTextureSinkFormat SinkFormat;
+	/** Texture sample queue. */
+	TSharedPtr<FMediaTextureSampleQueue, ESPMode::ThreadSafe> SampleQueue;
 
-	/** The mode that this sink is currently operating in. */
-	EMediaTextureSinkMode SinkMode;
+	/** Current size of the resource (in bytes).*/
+	SIZE_T Size;
 };
diff --git a/Engine/Source/Runtime/MediaAssets/Public/PlatformMediaSource.h b/Engine/Source/Runtime/MediaAssets/Public/PlatformMediaSource.h
--- a/Engine/Source/Runtime/MediaAssets/Public/PlatformMediaSource.h
+++ b/Engine/Source/Runtime/MediaAssets/Public/PlatformMediaSource.h
@@ -2,7 +2,9 @@
 
 #pragma once
 
-#include "CoreMinimal.h"
+#include "CoreTypes.h"
+#include "Containers/Map.h"
+#include "Containers/UnrealString.h"
 #include "UObject/ObjectMacros.h"
 #include "MediaSource.h"
 
diff --git a/Engine/Source/Runtime/MediaAssets/Public/StreamMediaSource.h b/Engine/Source/Runtime/MediaAssets/Public/StreamMediaSource.h
--- a/Engine/Source/Runtime/MediaAssets/Public/StreamMediaSource.h
+++ b/Engine/Source/Runtime/MediaAssets/Public/StreamMediaSource.h
@@ -2,9 +2,9 @@
 
 #pragma once
 
-#include "CoreMinimal.h"
-#include "UObject/ObjectMacros.h"
 #include "BaseMediaSource.h"
+#include "Containers/UnrealString.h"
+#include "UObject/ObjectMacros.h"
 
 #include "StreamMediaSource.generated.h"
 
diff --git a/Engine/Source/Runtime/MediaUtils/Private/MediaPlayerFacade.cpp b/Engine/Source/Runtime/MediaUtils/Private/MediaPlayerFacade.cpp
new file mode 100644
--- /dev/null
+++ b/Engine/Source/Runtime/MediaUtils/Private/MediaPlayerFacade.cpp
@@ -0,0 +1,965 @@
+// Copyright 1998-2016 Epic Games, Inc. All Rights Reserved.
+
+#include "MediaPlayerFacade.h"
+#include "MediaUtilsPrivate.h"
+
+#include "HAL/PlatformMath.h"
+#include "IMediaCache.h"
+#include "IMediaControls.h"
+#include "IMediaModule.h"
+#include "IMediaOptions.h"
+#include "IMediaPlayer.h"
+#include "IMediaPlayerFactory.h"
+#include "IMediaSamples.h"
+#include "IMediaTracks.h"
+#include "IMediaView.h"
+#include "Math/NumericLimits.h"
+#include "Misc/CoreMisc.h"
+#include "Misc/ScopeLock.h"
+#include "Modules/ModuleManager.h"
+
+#include "MediaSampleCache.h"
+
+
+/* Local helpers
+*****************************************************************************/
+
+namespace MediaPlayerFacade
+{
+	// @todo gmp: make these configurable in settings?
+
+	const FTimespan AudioPreroll = FTimespan::FromSeconds(1.0);
+	const FTimespan MetadataPreroll = FTimespan::FromSeconds(1.0);
+
+	const int32 MaxAudioSinkDepth = 512;
+	const int32 MaxCaptionSinkDepth = 256;
+	const int32 MaxMetadataSinkDepth = 256;
+	const int32 MaxSubtitleSinkDepth = 256;
+	const int32 MaxVideoSinkDepth = 8;
+}
+
+
+/* FMediaPlayerFacade structors
+*****************************************************************************/
+
+FMediaPlayerFacade::FMediaPlayerFacade()
+	: Cache(new FMediaSampleCache)
+	, LastRate(0.0f)
+{ }
+
+
+FMediaPlayerFacade::~FMediaPlayerFacade()
+{
+	if (Player.IsValid())
+	{
+		FScopeLock Lock(&CriticalSection);
+
+		Player->Close();
+		Player.Reset();
+	}
+
+	delete Cache;
+	Cache = nullptr;
+}
+
+
+/* FMediaPlayerFacade interface
+*****************************************************************************/
+
+void FMediaPlayerFacade::AddAudioSampleSink(const TSharedRef<TMediaSampleSink<IMediaAudioSample>, ESPMode::ThreadSafe>& SampleSink)
+{
+	AudioSampleSinks.Add(SampleSink);
+}
+
+
+void FMediaPlayerFacade::AddCaptionSampleSink(const TSharedRef<TMediaSampleSink<IMediaOverlaySample>, ESPMode::ThreadSafe>& SampleSink)
+{
+	CaptionSampleSinks.Add(SampleSink);
+}
+	
+
+void FMediaPlayerFacade::AddMetadataSampleSink(const TSharedRef<TMediaSampleSink<IMediaBinarySample>, ESPMode::ThreadSafe>& SampleSink)
+{
+	MetadataSampleSinks.Add(SampleSink);
+}
+
+
+void FMediaPlayerFacade::AddSubtitleSampleSink(const TSharedRef<TMediaSampleSink<IMediaOverlaySample>, ESPMode::ThreadSafe>& SampleSink)
+{
+	SubtitleSampleSinks.Add(SampleSink);
+}
+
+
+void FMediaPlayerFacade::AddVideoSampleSink(const TSharedRef<TMediaSampleSink<IMediaTextureSample>, ESPMode::ThreadSafe>& SampleSink)
+{
+	VideoSampleSinks.Add(SampleSink);
+}
+
+
+bool FMediaPlayerFacade::CanPause() const
+{
+	return Player.IsValid() && Player->GetControls().CanControl(EMediaControl::Pause);
+}
+
+
+bool FMediaPlayerFacade::CanPlayUrl(const FString& Url, const IMediaOptions* Options)
+{
+	IMediaModule* MediaModule = FModuleManager::LoadModulePtr<IMediaModule>("Media");
+
+	if (MediaModule == nullptr)
+	{
+		return false;
+	}
+
+	const FString RunningPlatformName(FPlatformProperties::IniPlatformName());
+	const TArray<IMediaPlayerFactory*>& PlayerFactories = MediaModule->GetPlayerFactories();
+
+	for (IMediaPlayerFactory* Factory : PlayerFactories)
+	{
+		if (Factory->SupportsPlatform(RunningPlatformName) && Factory->CanPlayUrl(Url, Options))
+		{
+			return true;
+		}
+	}
+
+	return false;
+}
+
+
+bool FMediaPlayerFacade::CanResume() const
+{
+	return Player.IsValid() && Player->GetControls().CanControl(EMediaControl::Resume);
+}
+
+
+bool FMediaPlayerFacade::CanScrub() const
+{
+	return Player.IsValid() && Player->GetControls().CanControl(EMediaControl::Scrub);
+}
+
+
+bool FMediaPlayerFacade::CanSeek() const
+{
+	return Player.IsValid() && Player->GetControls().CanControl(EMediaControl::Seek);
+}
+
+
+void FMediaPlayerFacade::Close()
+{
+	if (CurrentUrl.IsEmpty())
+	{
+		return;
+	}
+
+	if (Player.IsValid())
+	{
+		Player->Close();
+	}
+
+	Cache->Empty();
+	CurrentUrl.Empty();
+	LastRate = 0.0f;
+
+	FlushSinks();
+}
+
+
+uint32 FMediaPlayerFacade::GetAudioTrackChannels(int32 TrackIndex, int32 FormatIndex) const
+{
+	FMediaAudioTrackFormat Format;
+	return GetAudioTrackFormat(TrackIndex, FormatIndex, Format) ? Format.NumChannels : 0;
+}
+
+
+uint32 FMediaPlayerFacade::GetAudioTrackSampleRate(int32 TrackIndex, int32 FormatIndex) const
+{
+	FMediaAudioTrackFormat Format;
+	return GetAudioTrackFormat(TrackIndex, FormatIndex, Format) ? Format.SampleRate : 0;
+}
+
+
+FString FMediaPlayerFacade::GetAudioTrackType(int32 TrackIndex, int32 FormatIndex) const
+{
+	FMediaAudioTrackFormat Format;
+	return GetAudioTrackFormat(TrackIndex, FormatIndex, Format) ? Format.TypeName : FString();
+}
+
+
+FTimespan FMediaPlayerFacade::GetDuration() const
+{
+	return Player.IsValid() ? Player->GetControls().GetDuration() : FTimespan::Zero();
+}
+
+
+const FGuid& FMediaPlayerFacade::GetGuid()
+{
+	return PlayerGuid;
+}
+
+
+FString FMediaPlayerFacade::GetInfo() const
+{
+	return Player.IsValid() ? Player->GetInfo() : FString();
+}
+
+
+int32 FMediaPlayerFacade::GetNumTracks(EMediaTrackType TrackType) const
+{
+	return Player.IsValid() ? Player->GetTracks().GetNumTracks(TrackType) : 0;
+}
+
+
+int32 FMediaPlayerFacade::GetNumTrackFormats(EMediaTrackType TrackType, int32 TrackIndex) const
+{
+	return Player.IsValid() ? Player->GetTracks().GetNumTrackFormats(TrackType, TrackIndex) : 0;
+}
+
+
+FName FMediaPlayerFacade::GetPlayerName() const
+{
+	return Player.IsValid() ? Player->GetName() : NAME_None;
+}
+
+
+float FMediaPlayerFacade::GetRate() const
+{
+	return Player.IsValid() ? Player->GetControls().GetRate() : 0.0f;
+}
+
+
+int32 FMediaPlayerFacade::GetSelectedTrack(EMediaTrackType TrackType) const
+{
+	return Player.IsValid() ? Player->GetTracks().GetSelectedTrack((EMediaTrackType)TrackType) : INDEX_NONE;
+}
+
+
+FString FMediaPlayerFacade::GetStats() const
+{
+	return Player.IsValid() ? Player->GetStats() : FString();
+}
+
+
+TRangeSet<float> FMediaPlayerFacade::GetSupportedRates(bool Unthinned) const
+{
+	const EMediaRateThinning Thinning = Unthinned ? EMediaRateThinning::Unthinned : EMediaRateThinning::Thinned;
+
+	return Player.IsValid() ? Player->GetControls().GetSupportedRates(Thinning) : TRangeSet<float>();
+}
+
+
+FTimespan FMediaPlayerFacade::GetTime() const
+{
+	return Player.IsValid() ? Player->GetControls().GetTime() : FTimespan::Zero();
+}
+
+
+FText FMediaPlayerFacade::GetTrackDisplayName(EMediaTrackType TrackType, int32 TrackIndex) const
+{
+	return Player.IsValid() ? Player->GetTracks().GetTrackDisplayName((EMediaTrackType)TrackType, TrackIndex) : FText::GetEmpty();
+}
+
+
+int32 FMediaPlayerFacade::GetTrackFormat(EMediaTrackType TrackType, int32 TrackIndex) const
+{
+	return Player.IsValid() ? Player->GetTracks().GetTrackFormat((EMediaTrackType)TrackType, TrackIndex) : INDEX_NONE;
+}
+
+
+FString FMediaPlayerFacade::GetTrackLanguage(EMediaTrackType TrackType, int32 TrackIndex) const
+{
+	return Player.IsValid() ? Player->GetTracks().GetTrackLanguage((EMediaTrackType)TrackType, TrackIndex) : FString();
+}
+
+
+float FMediaPlayerFacade::GetVideoTrackAspectRatio(int32 TrackIndex, int32 FormatIndex) const
+{
+	FMediaVideoTrackFormat Format;
+	return (GetVideoTrackFormat(TrackIndex, FormatIndex, Format) && (Format.Dim.Y != 0)) ? ((float)(Format.Dim.X) / Format.Dim.Y) : 0.0f;
+}
+
+
+FIntPoint FMediaPlayerFacade::GetVideoTrackDimensions(int32 TrackIndex, int32 FormatIndex) const
+{
+	FMediaVideoTrackFormat Format;
+	return GetVideoTrackFormat(TrackIndex, FormatIndex, Format) ? Format.Dim : FIntPoint::ZeroValue;
+}
+
+
+float FMediaPlayerFacade::GetVideoTrackFrameRate(int32 TrackIndex, int32 FormatIndex) const
+{
+	FMediaVideoTrackFormat Format;
+	return GetVideoTrackFormat(TrackIndex, FormatIndex, Format) ? Format.FrameRate : 0.0f;
+}
+
+
+TRange<float> FMediaPlayerFacade::GetVideoTrackFrameRates(int32 TrackIndex, int32 FormatIndex) const
+{
+	FMediaVideoTrackFormat Format;
+	return GetVideoTrackFormat(TrackIndex, FormatIndex, Format) ? Format.FrameRates : TRange<float>::Empty();
+}
+
+
+FString FMediaPlayerFacade::GetVideoTrackType(int32 TrackIndex, int32 FormatIndex) const
+{
+	FMediaVideoTrackFormat Format;
+	return GetVideoTrackFormat(TrackIndex, FormatIndex, Format) ? Format.TypeName : FString();
+}
+
+
+bool FMediaPlayerFacade::GetViewField(float& OutHorizontal, float& OutVertical) const
+{
+	return Player.IsValid() ? Player->GetView().GetViewField(OutHorizontal, OutVertical) : false;
+}
+
+
+bool FMediaPlayerFacade::GetViewOrientation(FQuat& OutOrientation) const
+{
+	return Player.IsValid() ? Player->GetView().GetViewOrientation(OutOrientation) : false;
+}
+
+
+bool FMediaPlayerFacade::HasError() const
+{
+	return Player.IsValid() && (Player->GetControls().GetState() == EMediaState::Error);
+}
+
+
+bool FMediaPlayerFacade::IsBuffering() const
+{
+	return Player.IsValid() && EnumHasAnyFlags(Player->GetControls().GetStatus(), EMediaStatus::Buffering);
+}
+
+
+bool FMediaPlayerFacade::IsConnecting() const
+{
+	return Player.IsValid() && EnumHasAnyFlags(Player->GetControls().GetStatus(), EMediaStatus::Connecting);
+}
+
+
+bool FMediaPlayerFacade::IsLooping() const
+{
+	return Player.IsValid() && Player->GetControls().IsLooping();
+}
+
+
+bool FMediaPlayerFacade::IsPaused() const
+{
+	return Player.IsValid() && (Player->GetControls().GetState() == EMediaState::Paused);
+}
+
+
+bool FMediaPlayerFacade::IsPlaying() const
+{
+	return Player.IsValid() && (Player->GetControls().GetState() == EMediaState::Playing);
+}
+
+
+bool FMediaPlayerFacade::IsPreparing() const
+{
+	return Player.IsValid() && (Player->GetControls().GetState() == EMediaState::Preparing);
+}
+
+
+bool FMediaPlayerFacade::IsReady() const
+{
+	if (!Player.IsValid())
+	{
+		return false;
+	}
+
+	return ((Player->GetControls().GetState() != EMediaState::Closed) &&
+			(Player->GetControls().GetState() != EMediaState::Error) &&
+			(Player->GetControls().GetState() != EMediaState::Preparing));
+}
+
+
+bool FMediaPlayerFacade::Open(const FString& Url, const IMediaOptions* Options)
+{
+	if (IsRunningDedicatedServer())
+	{
+		return false;
+	}
+
+	Close();
+
+	if (Url.IsEmpty())
+	{
+		return false;
+	}
+
+	// find & initialize new player
+	TSharedPtr<IMediaPlayer, ESPMode::ThreadSafe> NewPlayer = GetPlayerForUrl(Url, Options);
+
+	if (NewPlayer != Player)
+	{
+		FScopeLock Lock(&CriticalSection);
+		Player = NewPlayer;
+	}
+
+	if (!Player.IsValid())
+	{
+		return false;
+	}
+
+	// update the Guid
+	Player->SetGuid(PlayerGuid);
+
+	CurrentUrl = Url;
+
+	// open the new media source
+	if (!Player->Open(Url, Options))
+	{
+		CurrentUrl.Empty();
+
+		return false;
+	}
+
+	return true;
+}
+
+
+void FMediaPlayerFacade::QueryCacheState(EMediaTrackType TrackType, EMediaCacheState State, TRangeSet<FTimespan>& OutTimeRanges) const
+{
+	if (!Player.IsValid())
+	{
+		return;
+	}
+
+	if (State == EMediaCacheState::Cached)
+	{
+		if (TrackType == EMediaTrackType::Audio)
+		{
+			Cache->GetCachedAudioSampleRanges(OutTimeRanges);
+		}
+		else if (TrackType == EMediaTrackType::Video)
+		{
+			Cache->GetCachedVideoSampleRanges(OutTimeRanges);
+		}
+	}
+	else
+	{
+		if (TrackType == EMediaTrackType::Video)
+		{
+			Player->GetCache().QueryCacheState(State, OutTimeRanges);
+		}
+	}
+}
+
+
+bool FMediaPlayerFacade::Seek(const FTimespan& Time)
+{
+	if (!Player.IsValid() || !Player->GetControls().Seek(Time))
+	{
+		return false;
+	}
+
+	return true;
+}
+
+
+bool FMediaPlayerFacade::SelectTrack(EMediaTrackType TrackType, int32 TrackIndex)
+{
+	if (!Player.IsValid() || !Player->GetTracks().SelectTrack((EMediaTrackType)TrackType, TrackIndex))
+	{
+		return false;
+	}
+	
+	FlushSinks();
+
+	return true;
+}
+
+
+void FMediaPlayerFacade::SetCacheWindow(FTimespan Ahead, FTimespan Behind)
+{
+	Cache->SetCacheWindow(Ahead, Behind);
+}
+
+
+void FMediaPlayerFacade::SetGuid(FGuid& Guid)
+{
+	PlayerGuid = Guid;
+}
+
+
+bool FMediaPlayerFacade::SetLooping(bool Looping)
+{
+	return Player.IsValid() && Player->GetControls().SetLooping(Looping);
+}
+
+
+bool FMediaPlayerFacade::SetRate(float Rate)
+{
+	if (!Player.IsValid() || !Player->GetControls().SetRate(Rate))
+	{
+		return false;
+	}
+
+	if ((LastRate * Rate) < 0.0f)
+	{
+		FlushSinks(); // direction change
+	}
+
+	return true;
+}
+
+
+bool FMediaPlayerFacade::SetTrackFormat(EMediaTrackType TrackType, int32 TrackIndex, int32 FormatIndex)
+{
+	return Player.IsValid() ? Player->GetTracks().SetTrackFormat((EMediaTrackType)TrackType, TrackIndex, FormatIndex) : false;
+}
+
+
+bool FMediaPlayerFacade::SetVideoTrackFrameRate(int32 TrackIndex, int32 FormatIndex, float FrameRate)
+{
+	return Player.IsValid() ? Player->GetTracks().SetVideoTrackFrameRate(TrackIndex, FormatIndex, FrameRate) : false;
+}
+
+
+bool FMediaPlayerFacade::SetViewField(float Horizontal, float Vertical, bool Absolute)
+{
+	return Player.IsValid() && Player->GetView().SetViewField(Horizontal, Vertical, Absolute);
+}
+
+
+bool FMediaPlayerFacade::SetViewOrientation(const FQuat& Orientation, bool Absolute)
+{
+	return Player.IsValid() && Player->GetView().SetViewOrientation(Orientation, Absolute);
+}
+
+
+bool FMediaPlayerFacade::SupportsRate(float Rate, bool Unthinned) const
+{
+	EMediaRateThinning Thinning = Unthinned ? EMediaRateThinning::Unthinned : EMediaRateThinning::Thinned;
+	return Player.IsValid() && Player->GetControls().GetSupportedRates(Thinning).Contains(Rate);
+}
+
+
+/* FMediaPlayerFacade implementation
+*****************************************************************************/
+
+void FMediaPlayerFacade::FlushSinks()
+{
+	UE_LOG(LogMediaUtils, Verbose, TEXT("PlayerFacade %llx: Flushing sinks"), this);
+
+	AudioSampleSinks.Flush();
+	CaptionSampleSinks.Flush();
+	MetadataSampleSinks.Flush();
+	SubtitleSampleSinks.Flush();
+	VideoSampleSinks.Flush();
+
+	if (Player.IsValid())
+	{
+		Player->GetSamples().FlushSamples();
+	}
+}
+
+
+bool FMediaPlayerFacade::GetAudioTrackFormat(int32 TrackIndex, int32 FormatIndex, FMediaAudioTrackFormat& OutFormat) const
+{
+	if (TrackIndex == INDEX_NONE)
+	{
+		TrackIndex = GetSelectedTrack(EMediaTrackType::Audio);
+	}
+
+	if (FormatIndex == INDEX_NONE)
+	{
+		FormatIndex = GetTrackFormat(EMediaTrackType::Audio, TrackIndex);
+	}
+
+	return (Player.IsValid() && Player->GetTracks().GetAudioTrackFormat(TrackIndex, FormatIndex, OutFormat));
+}
+
+
+TSharedPtr<IMediaPlayer, ESPMode::ThreadSafe> FMediaPlayerFacade::GetPlayerForUrl(const FString& Url, const IMediaOptions* Options)
+{
+	FName PlayerName;
+	
+	if (DesiredPlayerName != NAME_None)
+	{
+		PlayerName = DesiredPlayerName;
+	}
+	else if (Options != nullptr)
+	{
+		PlayerName = Options->GetDesiredPlayerName();
+	}
+	else
+	{
+		PlayerName = NAME_None;
+	}
+
+	// reuse existing player if requested
+	if (Player.IsValid() && (PlayerName == Player->GetName()))
+	{
+		return Player;
+	}
+
+	// load media module
+	IMediaModule* MediaModule = FModuleManager::LoadModulePtr<IMediaModule>("Media");
+
+	if (MediaModule == nullptr)
+	{
+		UE_LOG(LogMediaUtils, Error, TEXT("Failed to load Media module"));
+		return nullptr;
+	}
+
+	// try to create requested player
+	if (PlayerName != NAME_None)
+	{
+		IMediaPlayerFactory* Factory = MediaModule->GetPlayerFactory(PlayerName);
+
+		if (Factory == nullptr)
+		{
+			UE_LOG(LogMediaUtils, Error, TEXT("Could not find desired player %s for %s"), *PlayerName.ToString(), *Url);
+			return nullptr;
+		}
+
+		TSharedPtr<IMediaPlayer, ESPMode::ThreadSafe> NewPlayer = Factory->CreatePlayer(*this);
+
+		if (!NewPlayer.IsValid())
+		{
+			UE_LOG(LogMediaUtils, Error, TEXT("Failed to create desired player %s for %s"), *PlayerName.ToString(), *Url);
+			return nullptr;
+		}
+
+		return NewPlayer;
+	}
+
+	// try to reuse existing player
+	if (Player.IsValid())
+	{
+		IMediaPlayerFactory* Factory = MediaModule->GetPlayerFactory(Player->GetName());
+
+		if ((Factory != nullptr) && Factory->CanPlayUrl(Url, Options))
+		{
+			return Player;
+		}
+	}
+
+	const FString RunningPlatformName(FPlatformProperties::IniPlatformName());
+
+	// try to auto-select new player
+	const TArray<IMediaPlayerFactory*>& PlayerFactories = MediaModule->GetPlayerFactories();
+
+	for (IMediaPlayerFactory* Factory : PlayerFactories)
+	{
+		if (!Factory->SupportsPlatform(RunningPlatformName) || !Factory->CanPlayUrl(Url, Options))
+		{
+			continue;
+		}
+
+		TSharedPtr<IMediaPlayer, ESPMode::ThreadSafe> NewPlayer = Factory->CreatePlayer(*this);
+
+		if (NewPlayer.IsValid())
+		{
+			return NewPlayer;
+		}
+	}
+
+	// no suitable player found
+	if (PlayerFactories.Num() > 0)
+	{
+		UE_LOG(LogMediaUtils, Error, TEXT("Cannot play %s, because none of the enabled media player plug-ins support it:"), *Url);
+
+		for (IMediaPlayerFactory* Factory : PlayerFactories)
+		{
+			if (Factory->SupportsPlatform(RunningPlatformName))
+			{
+				UE_LOG(LogMediaUtils, Log, TEXT("| %s (URI scheme or file extension not supported)"), *Factory->GetName().ToString());
+			}
+			else
+			{
+				UE_LOG(LogMediaUtils, Log, TEXT("| %s (only available on %s, but not on %s)"), *Factory->GetName().ToString(), *FString::Join(Factory->GetSupportedPlatforms(), TEXT(", ")), *RunningPlatformName);
+			}	
+		}
+	}
+	else
+	{
+		UE_LOG(LogMediaUtils, Error, TEXT("Cannot play %s: no media player plug-ins are installed and enabled in this project"), *Url);
+	}
+
+	return nullptr;
+}
+
+
+bool FMediaPlayerFacade::GetVideoTrackFormat(int32 TrackIndex, int32 FormatIndex, FMediaVideoTrackFormat& OutFormat) const
+{
+	if (TrackIndex == INDEX_NONE)
+	{
+		TrackIndex = GetSelectedTrack(EMediaTrackType::Video);
+	}
+
+	if (FormatIndex == INDEX_NONE)
+	{
+		FormatIndex = GetTrackFormat(EMediaTrackType::Video, TrackIndex);
+	}
+
+	return (Player.IsValid() && Player->GetTracks().GetVideoTrackFormat(TrackIndex, FormatIndex, OutFormat));
+}
+
+
+void FMediaPlayerFacade::ProcessEvent(EMediaEvent Event)
+{
+	if (Event == EMediaEvent::TracksChanged)
+	{
+		SelectDefaultTracks();
+	}
+	else if ((Event == EMediaEvent::MediaOpened) || (Event == EMediaEvent::MediaOpenFailed))
+	{
+		if (Event == EMediaEvent::MediaOpenFailed)
+		{
+			CurrentUrl.Empty();
+		}
+
+		const FString MediaInfo = Player->GetInfo();
+
+		if (MediaInfo.IsEmpty())
+		{
+			UE_LOG(LogMediaUtils, Verbose, TEXT("PlayerFacade %llx: Media Info: n/a"), this);
+		}
+		else
+		{
+			UE_LOG(LogMediaUtils, Verbose, TEXT("PlayerFacade %llx: Media Info:\n%s"), this, *MediaInfo);
+		}
+	}
+
+	if ((Event == EMediaEvent::PlaybackEndReached) ||
+		(Event == EMediaEvent::SeekCompleted) ||
+		(Event == EMediaEvent::TracksChanged))
+	{
+		FlushSinks();
+	}
+
+	MediaEvent.Broadcast(Event);
+}
+
+
+void FMediaPlayerFacade::SelectDefaultTracks()
+{
+	IMediaTracks& Tracks = Player->GetTracks();
+
+	// @todo gmp: consider locale when selecting default media tracks
+
+	Tracks.SelectTrack(EMediaTrackType::Audio, 0);
+	Tracks.SelectTrack(EMediaTrackType::Caption, INDEX_NONE);
+	Tracks.SelectTrack(EMediaTrackType::Metadata, INDEX_NONE);
+	Tracks.SelectTrack(EMediaTrackType::Subtitle, INDEX_NONE);
+	Tracks.SelectTrack(EMediaTrackType::Video, 0);
+}
+
+
+/* IMediaClockSink interface
+*****************************************************************************/
+
+void FMediaPlayerFacade::TickFetch(FTimespan DeltaTime, FTimespan /*Timecode*/)
+{
+	// let the player generate samples & process events
+	if (Player.IsValid())
+	{
+		Player->TickFetch(DeltaTime);
+	}
+
+	// process deferred events
+	EMediaEvent Event;
+
+	while (QueuedEvents.Dequeue(Event))
+	{
+		ProcessEvent(Event);
+	}
+
+	if (!Player.IsValid())
+	{
+		return;
+	}
+
+	// get current play rate
+	float Rate = Player->GetControls().GetRate();
+
+	if (Rate == 0.0f)
+	{
+		Rate = LastRate;
+	}
+	else
+	{
+		LastRate = Rate;
+	}
+
+	// determine range of valid samples
+	TRange<FTimespan> TimeRange;
+
+	const FTimespan Time = Player->GetControls().GetTime();
+	
+	if (Rate > 0.0f)
+	{
+		TimeRange = TRange<FTimespan>::AtMost(Time);
+	}
+	else if (Rate < 0.0f)
+	{
+		TimeRange = TRange<FTimespan>::AtLeast(Time);
+	}
+	else
+	{
+		TimeRange = TRange<FTimespan>(Time);
+	}
+
+	// process samples in range
+	IMediaSamples& Samples = Player->GetSamples();
+
+	ProcessCaptionSamples(Samples, TimeRange);
+	ProcessSubtitleSamples(Samples, TimeRange);
+	ProcessVideoSamples(Samples, TimeRange);
+}
+
+
+void FMediaPlayerFacade::TickInput(FTimespan DeltaTime, FTimespan /*Timecode*/)
+{
+	if (Player.IsValid())
+	{
+		Player->TickInput(DeltaTime);
+	}
+}
+
+
+void FMediaPlayerFacade::TickOutput(FTimespan DeltaTime, FTimespan /*Timecode*/)
+{
+	if (!Player.IsValid())
+	{
+		return;
+	}
+
+	IMediaControls& Controls = Player->GetControls();
+	Cache->Tick(DeltaTime, Controls.GetRate(), Controls.GetTime());
+}
+
+
+/* IMediaTickable interface
+*****************************************************************************/
+
+void FMediaPlayerFacade::TickTickable()
+{
+	FScopeLock Lock(&CriticalSection);
+
+	if (!Player.IsValid())
+	{
+		return;
+	}
+
+	Player->TickAudio();
+
+	if (LastRate == 0.0f)
+	{
+		return;
+	}
+
+	// determine range of valid samples
+	TRange<FTimespan> AudioTimeRange;
+	TRange<FTimespan> MetadataTimeRange;
+
+	const FTimespan Time = Player->GetControls().GetTime();
+
+	if (LastRate > 0.0f)
+	{
+		AudioTimeRange = TRange<FTimespan>::Inclusive(FTimespan::MinValue(), Time + MediaPlayerFacade::AudioPreroll);
+		MetadataTimeRange = TRange<FTimespan>::Inclusive(FTimespan::MinValue(), Time + MediaPlayerFacade::MetadataPreroll);
+	}
+	else
+	{
+		AudioTimeRange = TRange<FTimespan>::Inclusive(Time - MediaPlayerFacade::AudioPreroll, FTimespan::MaxValue());
+		MetadataTimeRange = TRange<FTimespan>::Inclusive(Time - MediaPlayerFacade::MetadataPreroll, FTimespan::MaxValue());
+	}
+
+	// process samples in range
+	IMediaSamples& Samples = Player->GetSamples();
+	
+	ProcessAudioSamples(Samples, AudioTimeRange);
+	ProcessMetadataSamples(Samples, MetadataTimeRange);	
+}
+
+
+/* IMediaEventSink interface
+*****************************************************************************/
+
+void FMediaPlayerFacade::ReceiveMediaEvent(EMediaEvent Event)
+{
+	if (IsInGameThread())
+	{
+		ProcessEvent(Event);
+	}
+	else
+	{
+		QueuedEvents.Enqueue(Event);
+	}
+}
+
+
+/* FMediaPlayerFacade implementation
+*****************************************************************************/
+
+void FMediaPlayerFacade::ProcessAudioSamples(IMediaSamples& Samples, TRange<FTimespan> TimeRange)
+{
+	TSharedPtr<IMediaAudioSample, ESPMode::ThreadSafe> Sample;
+
+	while (Samples.FetchAudio(TimeRange, Sample))
+	{
+		if (!AudioSampleSinks.Enqueue(Sample, MediaPlayerFacade::MaxAudioSinkDepth))
+		{
+			UE_LOG(LogMediaUtils, VeryVerbose, TEXT("PlayerFacade %llx: Audio sample sink overflow"), this);
+		}
+	}
+}
+
+
+void FMediaPlayerFacade::ProcessCaptionSamples(IMediaSamples& Samples, TRange<FTimespan> TimeRange)
+{
+	TSharedPtr<IMediaOverlaySample, ESPMode::ThreadSafe> Sample;
+
+	while (Samples.FetchCaption(TimeRange, Sample))
+	{
+		if (!CaptionSampleSinks.Enqueue(Sample, MediaPlayerFacade::MaxCaptionSinkDepth))
+		{
+			UE_LOG(LogMediaUtils, VeryVerbose, TEXT("PlayerFacade %llx: Caption sample sink overflow"), this);
+		}
+	}
+}
+
+
+void FMediaPlayerFacade::ProcessMetadataSamples(IMediaSamples& Samples, TRange<FTimespan> TimeRange)
+{
+	TSharedPtr<IMediaBinarySample, ESPMode::ThreadSafe> Sample;
+
+	while (Samples.FetchMetadata(TimeRange, Sample))
+	{
+		if (!MetadataSampleSinks.Enqueue(Sample, MediaPlayerFacade::MaxMetadataSinkDepth))
+		{
+			UE_LOG(LogMediaUtils, VeryVerbose, TEXT("PlayerFacade %llx: Metadata sample sink overflow"), this);
+		}
+	}
+}
+
+
+void FMediaPlayerFacade::ProcessSubtitleSamples(IMediaSamples& Samples, TRange<FTimespan> TimeRange)
+{
+	TSharedPtr<IMediaOverlaySample, ESPMode::ThreadSafe> Sample;
+
+	while (Samples.FetchSubtitle(TimeRange, Sample))
+	{
+		if (!SubtitleSampleSinks.Enqueue(Sample, MediaPlayerFacade::MaxSubtitleSinkDepth))
+		{
+			UE_LOG(LogMediaUtils, VeryVerbose, TEXT("PlayerFacade %llx: Subtitle sample sink overflow"), this);
+		}
+	}
+}
+
+
+void FMediaPlayerFacade::ProcessVideoSamples(IMediaSamples& Samples, TRange<FTimespan> TimeRange)
+{
+	TSharedPtr<IMediaTextureSample, ESPMode::ThreadSafe> Sample;
+
+	while (Samples.FetchVideo(TimeRange, Sample))
+	{
+		if (!VideoSampleSinks.Enqueue(Sample, MediaPlayerFacade::MaxVideoSinkDepth))
+		{
+			UE_LOG(LogMediaUtils, VeryVerbose, TEXT("PlayerFacade %llx: Video sample sink overflow"), this);
+		}
+	}
+}
diff --git a/Engine/Source/Runtime/MediaUtils/Private/MediaSamples.cpp b/Engine/Source/Runtime/MediaUtils/Private/MediaSamples.cpp
new file mode 100644
--- /dev/null
+++ b/Engine/Source/Runtime/MediaUtils/Private/MediaSamples.cpp
@@ -0,0 +1,155 @@
+// Copyright 1998-2016 Epic Games, Inc. All Rights Reserved.
+
+#include "MediaSamples.h"
+
+#include "IMediaAudioSample.h"
+#include "IMediaBinarySample.h"
+#include "IMediaOverlaySample.h"
+#include "IMediaTextureSample.h"
+
+
+/* IMediaSamples interface
+*****************************************************************************/
+
+bool FMediaSamples::FetchAudio(TRange<FTimespan> TimeRange, TSharedPtr<IMediaAudioSample, ESPMode::ThreadSafe>& OutSample)
+{
+	TSharedPtr<IMediaAudioSample, ESPMode::ThreadSafe> Sample;
+
+	if (!AudioSampleQueue.Peek(Sample))
+	{
+		return false;
+	}
+
+	const FTimespan SampleTime = Sample->GetTime();
+
+	if (!TimeRange.Overlaps(TRange<FTimespan>(SampleTime, SampleTime + Sample->GetDuration())))
+	{
+		return false;
+	}
+
+	if (!AudioSampleQueue.Dequeue(Sample))
+	{
+		return false;
+	}
+
+	OutSample = Sample;
+
+	return true;
+}
+
+
+bool FMediaSamples::FetchCaption(TRange<FTimespan> TimeRange, TSharedPtr<IMediaOverlaySample, ESPMode::ThreadSafe>& OutSample)
+{
+	TSharedPtr<IMediaOverlaySample, ESPMode::ThreadSafe> Sample;
+
+	if (!CaptionSampleQueue.Peek(Sample))
+	{
+		return false;
+	}
+
+	const FTimespan SampleTime = Sample->GetTime();
+
+	if (!TimeRange.Overlaps(TRange<FTimespan>(SampleTime, SampleTime + Sample->GetDuration())))
+	{
+		return false;
+	}
+
+	if (!CaptionSampleQueue.Dequeue(Sample))
+	{
+		return false;
+	}
+
+	OutSample = Sample;
+
+	return true;
+}
+
+
+bool FMediaSamples::FetchMetadata(TRange<FTimespan> TimeRange, TSharedPtr<IMediaBinarySample, ESPMode::ThreadSafe>& OutSample)
+{
+	TSharedPtr<IMediaBinarySample, ESPMode::ThreadSafe> Sample;
+
+	if (!MetadataSampleQueue.Peek(Sample))
+	{
+		return false;
+	}
+
+	const FTimespan SampleTime = Sample->GetTime();
+
+	if (!TimeRange.Overlaps(TRange<FTimespan>(SampleTime, SampleTime + Sample->GetDuration())))
+	{
+		return false;
+	}
+
+	if (!MetadataSampleQueue.Dequeue(Sample))
+	{
+		return false;
+	}
+
+	OutSample = Sample;
+
+	return true;
+}
+
+
+bool FMediaSamples::FetchSubtitle(TRange<FTimespan> TimeRange, TSharedPtr<IMediaOverlaySample, ESPMode::ThreadSafe>& OutSample)
+{
+	TSharedPtr<IMediaOverlaySample, ESPMode::ThreadSafe> Sample;
+
+	if (!SubtitleSampleQueue.Peek(Sample))
+	{
+		return false;
+	}
+
+	const FTimespan SampleTime = Sample->GetTime();
+
+	if (!TimeRange.Overlaps(TRange<FTimespan>(SampleTime, SampleTime + Sample->GetDuration())))
+	{
+		return false;
+	}
+
+	if (!SubtitleSampleQueue.Dequeue(Sample))
+	{
+		return false;
+	}
+
+	OutSample = Sample;
+
+	return true;
+}
+
+
+bool FMediaSamples::FetchVideo(TRange<FTimespan> TimeRange, TSharedPtr<IMediaTextureSample, ESPMode::ThreadSafe>& OutSample)
+{
+	TSharedPtr<IMediaTextureSample, ESPMode::ThreadSafe> Sample;
+
+	if (!VideoSampleQueue.Peek(Sample))
+	{
+		return false;
+	}
+
+	const FTimespan SampleTime = Sample->GetTime();
+
+	if (!TimeRange.Overlaps(TRange<FTimespan>(SampleTime, SampleTime + Sample->GetDuration())))
+	{
+		return false;
+	}
+
+	if (!VideoSampleQueue.Dequeue(Sample))
+	{
+		return false;
+	}
+
+	OutSample = Sample;
+
+	return true;
+}
+
+
+void FMediaSamples::FlushSamples()
+{
+	AudioSampleQueue.RequestFlush();
+	MetadataSampleQueue.RequestFlush();
+	CaptionSampleQueue.RequestFlush();
+	VideoSampleQueue.RequestFlush();
+}
diff --git a/Engine/Source/Runtime/MediaUtils/Public/AppMediaTimeSource.h b/Engine/Source/Runtime/MediaUtils/Public/AppMediaTimeSource.h
new file mode 100644
--- /dev/null
+++ b/Engine/Source/Runtime/MediaUtils/Public/AppMediaTimeSource.h
@@ -0,0 +1,29 @@
+// Copyright 1998-2017 Epic Games, Inc. All Rights Reserved.
+
+#pragma once
+
+#include "IMediaTimeSource.h"
+#include "Misc/App.h"
+#include "Misc/Timespan.h"
+
+
+/**
+ * Implements the a media time source that derives its time from the application's global time.
+ */
+class FAppMediaTimeSource
+	: public IMediaTimeSource
+{
+public:
+
+	/** Virtual destructor. */
+	virtual ~FAppMediaTimeSource() { }
+
+public:
+
+	//~ IMediaTimeSource interface
+
+	virtual FTimespan GetTimecode() override
+	{
+		return FTimespan::FromSeconds(FApp::GetCurrentTime());
+	}
+};
diff --git a/Engine/Source/Runtime/MediaUtils/Public/MediaSampleQueue.h b/Engine/Source/Runtime/MediaUtils/Public/MediaSampleQueue.h
new file mode 100644
--- /dev/null
+++ b/Engine/Source/Runtime/MediaUtils/Public/MediaSampleQueue.h
@@ -0,0 +1,169 @@
+// Copyright 1998-2017 Epic Games, Inc. All Rights Reserved.
+
+#pragma once
+
+#include "Containers/Queue.h"
+#include "HAL/PlatformAtomics.h"
+#include "IMediaSamples.h"
+#include "Misc/App.h"
+#include "Misc/Timespan.h"
+#include "Templates/SharedPointer.h"
+
+#include "MediaSampleSink.h"
+#include "MediaSampleSource.h"
+
+
+/**
+ * Template for media sample queues.
+ */
+template<typename SampleType>
+class TMediaSampleQueue
+	: public TMediaSampleSink<SampleType>
+	, public TMediaSampleSource<SampleType>
+{
+public:
+
+	/** Default constructor. */
+	TMediaSampleQueue()
+		: NumSamples(0)
+		, PendingFlushes(0)
+	{ }
+
+	/** Virtual destructor. */
+	virtual ~TMediaSampleQueue() { }
+
+public:
+
+	/**
+	 * Get the number of samples in the queue.
+	 *
+	 * @return Number of samples.
+	 * @see Enqueue, Dequeue, Peek
+	 */
+	int32 Num() const
+	{
+		return NumSamples;
+	}
+
+public:
+
+	//~ TMediaSampleSource interface (to be called only from consumer thread)
+
+	virtual bool Dequeue(TSharedPtr<SampleType, ESPMode::ThreadSafe>& OutSample) override
+	{
+		DoPendingFlushes();
+
+		TSharedPtr<SampleType, ESPMode::ThreadSafe> Sample;
+
+		if (!Samples.Peek(Sample))
+		{
+			return false; // empty queue
+		}
+			
+		if (!Sample.IsValid())
+		{
+			return false; // pending flush
+		}
+
+		check(Samples.Dequeue(Sample));
+		check(Sample.IsValid());
+
+		FPlatformAtomics::InterlockedDecrement(&NumSamples);
+		check(NumSamples >= 0);
+
+		OutSample = Sample;
+
+		return true;
+	}
+
+	virtual bool Peek(TSharedPtr<SampleType, ESPMode::ThreadSafe>& OutSample) override
+	{
+		DoPendingFlushes();
+
+		TSharedPtr<SampleType, ESPMode::ThreadSafe> Sample;
+
+		if (!Samples.Peek(Sample))
+		{
+			return false; // empty queue
+		}
+
+		if (!Sample.IsValid())
+		{
+			return false; // pending flush
+		}
+
+		OutSample = Sample;
+
+		return true;
+	}
+
+public:
+
+	//~ TMediaSampleSink interface (to be called only from producer thread)
+
+	virtual bool Enqueue(const TSharedPtr<SampleType, ESPMode::ThreadSafe>& Sample) override
+	{
+		if (!Samples.Enqueue(Sample))
+		{
+			return false;
+		}
+
+		if (Sample.IsValid())
+		{
+			FPlatformAtomics::InterlockedIncrement(&NumSamples);
+		}
+
+		return true;
+	}
+
+	virtual void RequestFlush() override
+	{
+		Samples.Enqueue(nullptr); // insert flush marker
+		FPlatformAtomics::InterlockedIncrement(&PendingFlushes);
+	}
+
+protected:
+
+	/** Perform any pending flushes. */
+	void DoPendingFlushes()
+	{
+		TSharedPtr<SampleType, ESPMode::ThreadSafe> Sample;
+
+		while ((PendingFlushes > 0) && Samples.Dequeue(Sample))
+		{
+			if (Sample.IsValid())
+			{
+				FPlatformAtomics::InterlockedDecrement(&NumSamples);
+				check(NumSamples >= 0);
+			}
+			else
+			{
+				FPlatformAtomics::InterlockedDecrement(&PendingFlushes);
+			}
+		}
+	}
+
+private:
+
+	/** Number of samples in the queue. */
+	int32 NumSamples;
+
+	/** Number of pending flushes. */
+	int32 PendingFlushes;
+
+	/** Audio sample queue. */
+	TQueue<TSharedPtr<SampleType, ESPMode::ThreadSafe>, EQueueMode::Mpsc> Samples;
+};
+
+
+/** Type definition for audio sample queue. */
+typedef TMediaSampleQueue<class IMediaAudioSample> FMediaAudioSampleQueue;
+
+/** Type definition for binary sample queue. */
+typedef TMediaSampleQueue<class IMediaBinarySample> FMediaBinarySampleQueue;
+
+/** Type definition for overlay sample queue. */
+typedef TMediaSampleQueue<class IMediaOverlaySample> FMediaOverlaySampleQueue;
+
+/** Type definition for texture sample queue. */
+typedef TMediaSampleQueue<class IMediaTextureSample> FMediaTextureSampleQueue;
diff --git a/Engine/Source/Runtime/MediaUtils/Public/MediaSampleSinks.h b/Engine/Source/Runtime/MediaUtils/Public/MediaSampleSinks.h
new file mode 100644
--- /dev/null
+++ b/Engine/Source/Runtime/MediaUtils/Public/MediaSampleSinks.h
@@ -0,0 +1,122 @@
+// Copyright 1998-2016 Epic Games, Inc. All Rights Reserved.
+
+#pragma once
+
+#include "CoreTypes.h"
+#include "Containers/Array.h"
+#include "HAL/CriticalSection.h"
+#include "Misc/ScopeLock.h"
+#include "Templates/SharedPointer.h"
+
+#include "MediaSampleSink.h"
+
+
+/**
+ * Collection of media sample sinks.
+ *
+ * @param SampleType The type of media samples that the sinks process.
+ * @note All methods are thread-safe.
+ */
+template<typename SampleType>
+class TMediaSampleSinks
+{
+public:
+
+	/**
+	 * Add the given media sample sink to the collection.
+	 *
+	 * @param SampleSink The sink to add.
+	 * @see Remove
+	 */
+	void Add(const TSharedRef<TMediaSampleSink<SampleType>, ESPMode::ThreadSafe>& SampleSink)
+	{
+		FScopeLock Lock(&CriticalSection);
+		Sinks.AddUnique(SampleSink);
+	}
+
+	/**
+	 * Enqueue the given media samples to the registered sinks.
+	 *
+	 * This method will also remove expired sinks that haven't been removed yet.
+	 *
+	 * @param Sample The media sample to enqueue.
+	 * @param MaxQueueDepth The maximum depth of the sink queues before overflow.
+	 * @return true if the sample was enqueued to all sinks, false if one or more sinks overflowed.
+	 * @see Flush
+	 */
+	bool Enqueue(const TSharedPtr<SampleType, ESPMode::ThreadSafe>& Sample, int32 MaxDepth)
+	{
+		bool Overflowed = false;
+
+		FScopeLock Lock(&CriticalSection);
+
+		for (int32 SinkIndex = Sinks.Num() - 1; SinkIndex >= 0; --SinkIndex)
+		{
+			TSharedPtr<TMediaSampleSink<SampleType>, ESPMode::ThreadSafe> Sink = Sinks[SinkIndex].Pin();
+
+			if (Sink.IsValid())
+			{
+				if (Sink->Num() >= MaxDepth)
+				{
+					Overflowed = true;
+				}
+				else
+				{
+					Sink->Enqueue(Sample);
+				}
+			}
+			else
+			{
+				Sinks.RemoveAtSwap(SinkIndex);
+			}
+		}
+
+		return !Overflowed;
+	}
+
+	/**
+	 * Flush all registered sinks.
+	 *
+	 * This method will also remove expired sinks that haven't been removed yet.
+	 *
+	 * @see Enqueue
+	 */
+	void Flush()
+	{
+		FScopeLock Lock(&CriticalSection);
+
+		for (int32 SinkIndex = Sinks.Num() - 1; SinkIndex >= 0; --SinkIndex)
+		{
+			TSharedPtr<TMediaSampleSink<SampleType>, ESPMode::ThreadSafe> Sink = Sinks[SinkIndex].Pin();
+
+			if (Sink.IsValid())
+			{
+				Sink->RequestFlush();
+			}
+			else
+			{
+				Sinks.RemoveAtSwap(SinkIndex);
+			}
+		}
+	}
+
+	/**
+	 * Remove the given media sample sink from the collection.
+	 *
+	 * @param SampleSink The sink to remove.
+	 * @see Add
+	 */
+	void Remove(const TSharedRef<TMediaSampleSink<SampleType>, ESPMode::ThreadSafe>& SampleSink)
+	{
+		FScopeLock Lock(&CriticalSection);
+		Sinks.Remove(SampleSink);
+	}
+
+private:
+
+	/** Synchronizes access to Player. */
+	FCriticalSection CriticalSection;
+
+	/** The collection of registered sinks. */
+	TArray<TWeakPtr<TMediaSampleSink<SampleType>, ESPMode::ThreadSafe>> Sinks;
+};
diff --git a/Engine/Source/Runtime/MediaUtils/Public/MediaSamples.h b/Engine/Source/Runtime/MediaUtils/Public/MediaSamples.h
new file mode 100644
--- /dev/null
+++ b/Engine/Source/Runtime/MediaUtils/Public/MediaSamples.h
@@ -0,0 +1,162 @@
+// Copyright 1998-2017 Epic Games, Inc. All Rights Reserved.
+
+#pragma once
+
+#include "CoreTypes.h"
+#include "IMediaSamples.h"
+#include "Templates/SharedPointer.h"
+
+#include "MediaSampleQueue.h"
+
+class IMediaAudioSample;
+class IMediaBinarySample;
+class IMediaOverlaySample;
+class IMediaTextureSample;
+
+
+/**
+ * General purpose media sample queue.
+ */
+class MEDIAUTILS_API FMediaSamples
+	: public IMediaSamples
+{
+public:
+
+	/**
+	 * Add the given audio sample to the cache.
+	 *
+	 * @param Sample The sample to add.
+	 * @see AddCaption, AddMetadata, AddSubtitle, AddVideo, NumAudio
+	 */
+	void AddAudio(const TSharedRef<IMediaAudioSample, ESPMode::ThreadSafe>& Sample)
+	{
+		AudioSampleQueue.Enqueue(Sample);
+	}
+
+	/**
+	 * Add the given caption sample to the cache.
+	 *
+	 * @param Sample The sample to add.
+	 * @see AddAudio, AddMetadata, AddSubtitle, AddVideo, NumCaption
+	 */
+	void AddCaption(const TSharedRef<IMediaOverlaySample, ESPMode::ThreadSafe>& Sample)
+	{
+		CaptionSampleQueue.Enqueue(Sample);
+	}
+
+	/**
+	 * Add the given audio sample to the cache.
+	 *
+	 * @param Sample The sample to add.
+	 * @see AddAudio, AddCaption, AddSubtitle, AddVideo, NumMetadata
+	 */
+	void AddMetadata(const TSharedRef<IMediaBinarySample, ESPMode::ThreadSafe>& Sample)
+	{
+		MetadataSampleQueue.Enqueue(Sample);
+	}
+
+	/**
+	 * Add the given subtitle sample to the cache.
+	 *
+	 * @param Sample The sample to add.
+	 * @see AddAudio, AddCaption, AddMetadata, AddVideo, NumSubtitle
+	 */
+	void AddSubtitle(const TSharedRef<IMediaOverlaySample, ESPMode::ThreadSafe>& Sample)
+	{
+		SubtitleSampleQueue.Enqueue(Sample);
+	}
+
+	/**
+	 * Add the given audio sample to the cache.
+	 *
+	 * @param Sample The sample to add.
+	 * @see AddAudio, AddCaption, AddMetadata, AddSubtitle, NumVideo
+	 */
+	void AddVideo(const TSharedRef<IMediaTextureSample, ESPMode::ThreadSafe>& Sample)
+	{
+		VideoSampleQueue.Enqueue(Sample);
+	}
+
+	/**
+	 * Get the number of queued audio samples.
+	 *
+	 * @return Number of samples.
+	 * @see AddAudio, NumCaption, NumMetadata, NumSubtitle, NumVideo
+	 */
+	int32 NumAudio() const
+	{
+		return AudioSampleQueue.Num();
+	}
+
+	/**
+	 * Get the number of queued caption samples.
+	 *
+	 * @return Number of samples.
+	 * @see AddCaption, NumAudio, NumMetadata, NumSubtitle, NumVideo
+	 */
+	int32 NumCaption() const
+	{
+		return CaptionSampleQueue.Num();
+	}
+
+	/**
+	 * Get the number of queued metadata samples.
+	 *
+	 * @return Number of samples.
+	 * @see AddMetadata, NumAudio, NumCaption, NumSubtitle, NumVideo
+	 */
+	int32 NumMetadataSamples() const
+	{
+		return MetadataSampleQueue.Num();
+	}
+
+	/**
+	 * Get the number of queued subtitle samples.
+	 *
+	 * @return Number of samples.
+	 * @see AddSubtitle, NumAudio, NumCaption, NumMetadata, NumVideo
+	 */
+	int32 NumSubtitleSamples() const
+	{
+		return SubtitleSampleQueue.Num();
+	}
+
+	/**
+	 * Get the number of queued video samples.
+	 *
+	 * @return Number of samples.
+	 * @see AddVideo, NumAudio, NumCaption, NumMetadata, NumSubtitle
+	 */
+	int32 NumVideoSamples() const
+	{
+		return VideoSampleQueue.Num();
+	}
+
+public:
+
+	//~ IMediaSamples interface
+
+	virtual bool FetchAudio(TRange<FTimespan> TimeRange, TSharedPtr<IMediaAudioSample, ESPMode::ThreadSafe>& OutSample) override;
+	virtual bool FetchCaption(TRange<FTimespan> TimeRange, TSharedPtr<IMediaOverlaySample, ESPMode::ThreadSafe>& OutSample) override;
+	virtual bool FetchMetadata(TRange<FTimespan> TimeRange, TSharedPtr<IMediaBinarySample, ESPMode::ThreadSafe>& OutSample) override;
+	virtual bool FetchSubtitle(TRange<FTimespan> TimeRange, TSharedPtr<IMediaOverlaySample, ESPMode::ThreadSafe>& OutSample) override;
+	virtual bool FetchVideo(TRange<FTimespan> TimeRange, TSharedPtr<IMediaTextureSample, ESPMode::ThreadSafe>& OutSample) override;
+	virtual void FlushSamples() override;
+
+private:
+
+	/** Audio sample queue. */
+	TMediaSampleQueue<IMediaAudioSample> AudioSampleQueue;
+
+	/** Caption sample queue. */
+	TMediaSampleQueue<IMediaOverlaySample> CaptionSampleQueue;
+
+	/** Metadata sample queue. */
+	TMediaSampleQueue<IMediaBinarySample> MetadataSampleQueue;
+
+	/** Subtitle sample queue. */
+	TMediaSampleQueue<IMediaOverlaySample> SubtitleSampleQueue;
+
+	/** Video sample queue. */
+	TMediaSampleQueue<IMediaTextureSample> VideoSampleQueue;
+};
diff --git a/Engine/Source/Runtime/Messaging/Private/Bridge/MessageBridge.cpp b/Engine/Source/Runtime/Messaging/Private/Bridge/MessageBridge.cpp
--- a/Engine/Source/Runtime/Messaging/Private/Bridge/MessageBridge.cpp
+++ b/Engine/Source/Runtime/Messaging/Private/Bridge/MessageBridge.cpp
@@ -2,6 +2,10 @@
 
 #include "Bridge/MessageBridge.h"
 
+#include "IMessageBus.h"
+#include "IMessageSubscription.h"
+#include "IMessageTransport.h"
+
 
 /* FMessageBridge structors
  *****************************************************************************/
@@ -18,14 +22,12 @@ FMessageBridge::FMessageBridge(
 	, Transport(InTransport)
 {
 	Bus->OnShutdown().AddRaw(this, &FMessageBridge::HandleMessageBusShutdown);
-	Transport->OnNodeLost().BindRaw(this, &FMessageBridge::HandleTransportNodeLost);
-	Transport->OnMessageReceived().BindRaw(this, &FMessageBridge::HandleTransportMessageReceived);
 }
 
 
 FMessageBridge::~FMessageBridge()
 {
-	Shutdown();
+	Disable();
 
 	if (Bus.IsValid())
 	{
@@ -76,7 +78,7 @@ void FMessageBridge::Enable()
 	}
 
 	// enable subscription & transport
-	if (!Transport->StartTransport())
+	if (!Transport->StartTransport(*this))
 	{
 		return;
 	}
@@ -96,10 +98,40 @@ void FMessageBridge::Enable()
 }
 
 
-/* IReceiveMessages interface
+bool FMessageBridge::IsEnabled() const
+{
+	return Enabled;
+}
+
+
+/* IMessageReceiver interface
  *****************************************************************************/
 
-void FMessageBridge::ReceiveMessage(const IMessageContextRef& Context)
+FName FMessageBridge::GetDebugName() const
+{
+	return *FString::Printf(TEXT("FMessageBridge (%s)"), *Transport->GetDebugName().ToString());
+}
+
+
+const FGuid& FMessageBridge::GetRecipientId() const
+{
+	return Id;
+}
+
+
+ENamedThreads::Type FMessageBridge::GetRecipientThread() const
+{
+	return ENamedThreads::AnyThread;
+}
+
+
+bool FMessageBridge::IsLocal() const
+{
+	return false;
+}
+
+
+void FMessageBridge::ReceiveMessage(const TSharedRef<IMessageContext, ESPMode::ThreadSafe>& Context)
 {
 	if (!Enabled)
 	{
@@ -124,41 +156,49 @@ void FMessageBridge::ReceiveMessage(const IMessageContextRef& Context)
 }
 
 
-/* ISendMessages interface
+/* IMessageSender interface
  *****************************************************************************/
 
+FMessageAddress FMessageBridge::GetSenderAddress()
+{
+	return Address;
+}
+
+
 void FMessageBridge::NotifyMessageError(const TSharedRef<IMessageContext, ESPMode::ThreadSafe>& Context, const FString& Error)
 {
 	// deprecated
 }
 
 
-/* FMessageBridge implementation
+/* IMessageTransportHandler interface
  *****************************************************************************/
 
-void FMessageBridge::Shutdown()
+void FMessageBridge::DiscoverTransportNode(const FGuid& NodeId)
 {
-	Disable();
-
-	if (Transport.IsValid())
-	{
-		Transport->OnMessageReceived().Unbind();
-		Transport->OnNodeLost().Unbind();
-	}
+	// do nothing (address book is updated in ReceiveTransportMessage)
 }
 
 
-/* FMessageBridge callbacks
- *****************************************************************************/
-
-void FMessageBridge::HandleMessageBusShutdown()
+void FMessageBridge::ForgetTransportNode(const FGuid& NodeId)
 {
-	Shutdown();
-	Bus.Reset();
+	TArray<FMessageAddress> RemovedAddresses;
+
+	// update address book
+	AddressBook.RemoveNode(NodeId, RemovedAddresses);
+
+	// unregister endpoints
+	if (Bus.IsValid())
+	{
+		for (const auto& RemovedAddress : RemovedAddresses)
+		{
+			Bus->Unregister(RemovedAddress);
+		}
+	}
 }
 
 
-void FMessageBridge::HandleTransportMessageReceived(const TSharedRef<IMessageContext, ESPMode::ThreadSafe>& Context, const FGuid& NodeId)
+void FMessageBridge::ReceiveTransportMessage(const TSharedRef<IMessageContext, ESPMode::ThreadSafe>& Context, const FGuid& NodeId)
 {
 	if (!Enabled || !Bus.IsValid())
 	{
@@ -183,19 +223,11 @@ void FMessageBridge::HandleTransportMessageReceived(const TSharedRef<IMessageCon
 }
 
 
-void FMessageBridge::HandleTransportNodeLost(const FGuid& NodeId)
-{
-	TArray<FMessageAddress> RemovedAddresses;
-
-	// update address book
-	AddressBook.RemoveNode(NodeId, RemovedAddresses);
+/* FMessageBridge callbacks
+ *****************************************************************************/
 
-	// unregister endpoints
-	if (Bus.IsValid())
-	{
-		for (const auto& RemovedAddress : RemovedAddresses)
-		{
-			Bus->Unregister(RemovedAddress);
-		}
-	}
+void FMessageBridge::HandleMessageBusShutdown()
+{
+	Disable();
+	Bus.Reset();
 }
diff --git a/Engine/Source/Runtime/Messaging/Private/Bridge/MessageBridge.h b/Engine/Source/Runtime/Messaging/Private/Bridge/MessageBridge.h
--- a/Engine/Source/Runtime/Messaging/Private/Bridge/MessageBridge.h
+++ b/Engine/Source/Runtime/Messaging/Private/Bridge/MessageBridge.h
@@ -2,18 +2,20 @@
 
 #pragma once
 
-#include "CoreMinimal.h"
-#include "Misc/Guid.h"
-#include "IMessageContext.h"
-#include "IMessageTransport.h"
-#include "IMessageSubscription.h"
-#include "IMessageBus.h"
-#include "Bridge/MessageAddressBook.h"
 #include "IMessageBridge.h"
+#include "IMessageContext.h"
 #include "IMessageReceiver.h"
 #include "IMessageSender.h"
+#include "IMessageTransportHandler.h"
+#include "Misc/Guid.h"
+#include "Templates/SharedPointer.h"
+
+#include "Bridge/MessageAddressBook.h"
+
+class IMessageBus;
+class IMessageSubscription;
+class IMessageTransport;
 
-class Error;
 
 /**
  * Implements a message bridge.
@@ -34,6 +36,7 @@ class FMessageBridge
 	, public IMessageBridge
 	, public IMessageReceiver
 	, public IMessageSender
+	, protected IMessageTransportHandler
 {
 public:
 
@@ -47,7 +50,8 @@ public:
 	FMessageBridge(
 		const FMessageAddress InAddress,
 		const TSharedRef<IMessageBus, ESPMode::ThreadSafe>& InBus,
-		const TSharedRef<IMessageTransport, ESPMode::ThreadSafe>& InTransport);
+		const TSharedRef<IMessageTransport, ESPMode::ThreadSafe>& InTransport
+	);
 
 	/** Virtual destructor. */
 	virtual ~FMessageBridge();
@@ -58,65 +62,38 @@ public:
 
 	virtual void Disable() override;
 	virtual void Enable() override;
-
-	virtual bool IsEnabled() const override
-	{
-		return Enabled;
-	}
+	virtual bool IsEnabled() const override;
 
 public:
 
-	//~ IReceiveMessages interface
+	//~ IMessageReceiver interface
 
-	virtual FName GetDebugName() const override
-	{
-		return *FString::Printf(TEXT("FMessageBridge (%s)"), *Transport->GetDebugName().ToString());
-	}
-
-	virtual const FGuid& GetRecipientId() const override
-	{
-		return Id;
-	}
-
-	virtual ENamedThreads::Type GetRecipientThread() const override
-	{
-		return ENamedThreads::AnyThread;
-	}
-
-	virtual bool IsLocal() const override
-	{
-		return false;
-	}
-
-	virtual void ReceiveMessage(const IMessageContextRef& Context) override;
+	virtual FName GetDebugName() const override;
+	virtual const FGuid& GetRecipientId() const override;
+	virtual ENamedThreads::Type GetRecipientThread() const override;
+	virtual bool IsLocal() const override;
+	virtual void ReceiveMessage(const TSharedRef<IMessageContext, ESPMode::ThreadSafe>& Context) override;
 
 public:
 
-	//~ ISendMessages interface
-
-	virtual FMessageAddress GetSenderAddress() override
-	{
-		return Address;
-	}
+	//~ IMessageSender interface
 
+	virtual FMessageAddress GetSenderAddress() override;
 	virtual void NotifyMessageError(const TSharedRef<IMessageContext, ESPMode::ThreadSafe>& Context, const FString& Error) override;
 
 protected:
 
-	/** Shuts down the bridge. */
-	void Shutdown();
+	//~ IMessageTransportHandler interface
+
+	virtual void DiscoverTransportNode(const FGuid& NodeId) override;
+	virtual void ForgetTransportNode(const FGuid& NodeId) override;
+	virtual void ReceiveTransportMessage(const TSharedRef<IMessageContext, ESPMode::ThreadSafe>& Context, const FGuid& NodeId) override;
 
 private:
 
 	/** Callback for message bus shutdowns. */
 	void HandleMessageBusShutdown();
 
-	/** Callback for messages received from the transport layer. */
-	void HandleTransportMessageReceived(const TSharedRef<IMessageContext, ESPMode::ThreadSafe>& Envelope, const FGuid& NodeId);
-
-	/** Callback for lost remote nodes. */
-	void HandleTransportNodeLost(const FGuid& NodeId);
-
 private:
 
 	/** Holds the bridge's address. */
diff --git a/Engine/Source/Runtime/Messaging/Private/Bus/MessageContext.cpp b/Engine/Source/Runtime/Messaging/Private/Bus/MessageContext.cpp
--- a/Engine/Source/Runtime/Messaging/Private/Bus/MessageContext.cpp
+++ b/Engine/Source/Runtime/Messaging/Private/Bus/MessageContext.cpp
@@ -78,7 +78,7 @@ const TWeakObjectPtr<UScriptStruct>& FMessageContext::GetMessageTypeInfo() const
 }
 
 
-IMessageContextPtr FMessageContext::GetOriginalContext() const
+TSharedPtr<IMessageContext, ESPMode::ThreadSafe> FMessageContext::GetOriginalContext() const
 {
 	return OriginalContext;
 }
diff --git a/Engine/Source/Runtime/Messaging/Private/Bus/MessageContext.h b/Engine/Source/Runtime/Messaging/Private/Bus/MessageContext.h
--- a/Engine/Source/Runtime/Messaging/Private/Bus/MessageContext.h
+++ b/Engine/Source/Runtime/Messaging/Private/Bus/MessageContext.h
@@ -74,7 +74,7 @@ public:
 	 * @param InForwarderThread The name of the thread from which the message was forwarded.
 	 */
 	FMessageContext(
-		const IMessageContextRef& InContext,
+		const TSharedRef<IMessageContext, ESPMode::ThreadSafe>& InContext,
 		const FMessageAddress& InForwarder,
 		const TArray<FMessageAddress>& NewRecipients,
 		EMessageScope NewScope,
@@ -102,7 +102,7 @@ public:
 	virtual const FDateTime& GetExpiration() const override;
 	virtual const void* GetMessage() const override;
 	virtual const TWeakObjectPtr<UScriptStruct>& GetMessageTypeInfo() const override;
-	virtual IMessageContextPtr GetOriginalContext() const override;
+	virtual TSharedPtr<IMessageContext, ESPMode::ThreadSafe> GetOriginalContext() const override;
 	virtual const TArray<FMessageAddress>& GetRecipients() const override;
 	virtual EMessageScope GetScope() const override;
 	virtual const FMessageAddress& GetSender() const override;
@@ -125,7 +125,7 @@ private:
 	void* Message;
 
 	/** Holds the original message context. */
-	IMessageContextPtr OriginalContext;
+	TSharedPtr<IMessageContext, ESPMode::ThreadSafe> OriginalContext;
 
 	/** Holds the message recipients. */
 	TArray<FMessageAddress> Recipients;
diff --git a/Engine/Source/Runtime/Messaging/Private/Bus/MessageRouter.cpp b/Engine/Source/Runtime/Messaging/Private/Bus/MessageRouter.cpp
--- a/Engine/Source/Runtime/Messaging/Private/Bus/MessageRouter.cpp
+++ b/Engine/Source/Runtime/Messaging/Private/Bus/MessageRouter.cpp
@@ -46,7 +46,6 @@ uint32 FMessageRouter::Run()
 		if (WorkEvent->Wait(CalculateWaitTime()))
 		{
 			CurrentTime = FDateTime::UtcNow();
-
 			CommandDelegate Command;
 
 			while (Commands.Dequeue(Command))
@@ -67,8 +66,8 @@ uint32 FMessageRouter::Run()
 void FMessageRouter::Stop()
 {
 	Tracer->Stop();
-
 	Stopping = true;
+	WorkEvent->Trigger();
 }
 
 
@@ -113,7 +112,7 @@ FTimespan FMessageRouter::CalculateWaitTime()
 }
 
 
-void FMessageRouter::DispatchMessage(const IMessageContextRef& Context)
+void FMessageRouter::DispatchMessage(const TSharedRef<IMessageContext, ESPMode::ThreadSafe>& Context)
 {
 	if (Context->IsValid())
 	{
@@ -167,7 +166,7 @@ void FMessageRouter::DispatchMessage(const IMessageContextRef& Context)
 
 void FMessageRouter::FilterSubscriptions(
 	TArray<TSharedPtr<IMessageSubscription, ESPMode::ThreadSafe>>& Subscriptions,
-	const IMessageContextRef& Context,
+	const TSharedRef<IMessageContext, ESPMode::ThreadSafe>& Context,
 	TArray<TSharedPtr<IMessageReceiver, ESPMode::ThreadSafe>>& OutRecipients
 )
 {
@@ -315,7 +314,7 @@ void FMessageRouter::HandleRemoveSubscriber(TWeakPtr<IMessageReceiver, ESPMode::
 }
 
 
-void FMessageRouter::HandleRouteMessage(IMessageContextRef Context)
+void FMessageRouter::HandleRouteMessage(TSharedRef<IMessageContext, ESPMode::ThreadSafe> Context)
 {
 	Tracer->TraceRoutedMessage(Context);
 
diff --git a/Engine/Source/Runtime/Messaging/Private/Bus/MessageRouter.h b/Engine/Source/Runtime/Messaging/Private/Bus/MessageRouter.h
--- a/Engine/Source/Runtime/Messaging/Private/Bus/MessageRouter.h
+++ b/Engine/Source/Runtime/Messaging/Private/Bus/MessageRouter.h
@@ -181,7 +181,7 @@ private:
 	struct FDelayedMessage
 	{
 		/** Holds the context of the delayed message. */
-		IMessageContextPtr Context;
+		TSharedPtr<IMessageContext, ESPMode::ThreadSafe> Context;
 
 		/** Holds a sequence number used by the delayed message queue. */
 		int64 Sequence;
diff --git a/Engine/Source/Runtime/Messaging/Private/Bus/MessageTracer.cpp b/Engine/Source/Runtime/Messaging/Private/Bus/MessageTracer.cpp
--- a/Engine/Source/Runtime/Messaging/Private/Bus/MessageTracer.cpp
+++ b/Engine/Source/Runtime/Messaging/Private/Bus/MessageTracer.cpp
@@ -98,7 +98,7 @@ void FMessageTracer::TraceAddedSubscription(const TSharedRef<IMessageSubscriptio
 }
 
 
-void FMessageTracer::TraceDispatchedMessage(const IMessageContextRef& Context, const TSharedRef<IMessageReceiver, ESPMode::ThreadSafe>& Recipient, bool Async)
+void FMessageTracer::TraceDispatchedMessage(const TSharedRef<IMessageContext, ESPMode::ThreadSafe>& Context, const TSharedRef<IMessageReceiver, ESPMode::ThreadSafe>& Recipient, bool Async)
 {
 	if (!Running)
 	{
@@ -142,7 +142,7 @@ void FMessageTracer::TraceDispatchedMessage(const IMessageContextRef& Context, c
 }
 
 
-void FMessageTracer::TraceHandledMessage(const IMessageContextRef& Context, const TSharedRef<IMessageReceiver, ESPMode::ThreadSafe>& Recipient)
+void FMessageTracer::TraceHandledMessage(const TSharedRef<IMessageContext, ESPMode::ThreadSafe>& Context, const TSharedRef<IMessageReceiver, ESPMode::ThreadSafe>& Recipient)
 {
 	if (!Running)
 	{
@@ -178,7 +178,7 @@ void FMessageTracer::TraceHandledMessage(const IMessageContextRef& Context, cons
 }
 
 
-void FMessageTracer::TraceInterceptedMessage(const IMessageContextRef& Context, const TSharedRef<IMessageInterceptor, ESPMode::ThreadSafe>& Interceptor)
+void FMessageTracer::TraceInterceptedMessage(const TSharedRef<IMessageContext, ESPMode::ThreadSafe>& Context, const TSharedRef<IMessageInterceptor, ESPMode::ThreadSafe>& Interceptor)
 {
 	if (!Running)
 	{
@@ -267,7 +267,7 @@ void FMessageTracer::TraceRemovedSubscription(const TSharedRef<IMessageSubscript
 }
 
 
-void FMessageTracer::TraceRoutedMessage(const IMessageContextRef& Context)
+void FMessageTracer::TraceRoutedMessage(const TSharedRef<IMessageContext, ESPMode::ThreadSafe>& Context)
 {
 	if (!Running)
 	{
@@ -294,7 +294,7 @@ void FMessageTracer::TraceRoutedMessage(const IMessageContextRef& Context)
 }
 
 
-void FMessageTracer::TraceSentMessage(const IMessageContextRef& Context)
+void FMessageTracer::TraceSentMessage(const TSharedRef<IMessageContext, ESPMode::ThreadSafe>& Context)
 {
 	if (!Running)
 	{
@@ -488,7 +488,7 @@ void FMessageTracer::ResetMessages()
 }
 
 
-bool FMessageTracer::ShouldBreak(const IMessageContextRef& Context) const
+bool FMessageTracer::ShouldBreak(const TSharedRef<IMessageContext, ESPMode::ThreadSafe>& Context) const
 {
 	if (Breaking)
 	{
diff --git a/Engine/Source/Runtime/Messaging/Private/MessagingModule.cpp b/Engine/Source/Runtime/Messaging/Private/MessagingModule.cpp
--- a/Engine/Source/Runtime/Messaging/Private/MessagingModule.cpp
+++ b/Engine/Source/Runtime/Messaging/Private/MessagingModule.cpp
@@ -62,7 +62,7 @@ public:
 
 	//~ IMessagingModule interface
 
-	virtual TSharedPtr<IMessageBridge, ESPMode::ThreadSafe> CreateBridge(const FMessageAddress& Address, const TSharedRef<IMessageBus, ESPMode::ThreadSafe>& Bus, const IMessageTransportRef& Transport) override
+	virtual TSharedPtr<IMessageBridge, ESPMode::ThreadSafe> CreateBridge(const FMessageAddress& Address, const TSharedRef<IMessageBus, ESPMode::ThreadSafe>& Bus, const TSharedRef<IMessageTransport, ESPMode::ThreadSafe>& Transport) override
 	{
 		return MakeShareable(new FMessageBridge(Address, Bus, Transport));
 	}
diff --git a/Engine/Source/Runtime/Messaging/Public/IAuthorizeMessageRecipients.h b/Engine/Source/Runtime/Messaging/Public/IAuthorizeMessageRecipients.h
--- a/Engine/Source/Runtime/Messaging/Public/IAuthorizeMessageRecipients.h
+++ b/Engine/Source/Runtime/Messaging/Public/IAuthorizeMessageRecipients.h
@@ -2,11 +2,16 @@
 
 #pragma once
 
-#include "CoreMinimal.h"
-#include "IMessageInterceptor.h"
+#include "CoreTypes.h"
+#include "Templates/SharedPointer.h"
 
+class FName;
+class IMessageInterceptor;
 class IMessageReceiver;
 
+struct FMessageAddress;
+
+
 /**
  * Interface for classes that authorize message subscriptions.
  */
diff --git a/Engine/Source/Runtime/Messaging/Public/IMessageAttachment.h b/Engine/Source/Runtime/Messaging/Public/IMessageAttachment.h
--- a/Engine/Source/Runtime/Messaging/Public/IMessageAttachment.h
+++ b/Engine/Source/Runtime/Messaging/Public/IMessageAttachment.h
@@ -2,7 +2,8 @@
 
 #pragma once
 
-#include "CoreMinimal.h"
+class FArchive;
+
 
 /**
  * Interface for message attachments.
diff --git a/Engine/Source/Runtime/Messaging/Public/IMessageBus.h b/Engine/Source/Runtime/Messaging/Public/IMessageBus.h
--- a/Engine/Source/Runtime/Messaging/Public/IMessageBus.h
+++ b/Engine/Source/Runtime/Messaging/Public/IMessageBus.h
@@ -2,15 +2,26 @@
 
 #pragma once
 
-#include "CoreMinimal.h"
-#include "IMessageContext.h"
-#include "IMessageInterceptor.h"
-#include "IMessageTracer.h"
+#include "Containers/Array.h"
+#include "Math/Range.h"
+#include "Templates/SharedPointer.h"
 
+class FName;
 class IMessageAttachment;
+class IMessageContext;
+class IMessageInterceptor;
 class IMessageReceiver;
 class IMessageSender;
 class IMessageSubscription;
+class IMessageTracer;
+class UScriptStruct;
+
+enum class EMessageScope : uint8;
+
+struct FDateTime;
+struct FMessageAddress;
+struct FTimespan;
+
 
 /** Delegate type for message bus shutdowns. */
 DECLARE_MULTICAST_DELEGATE(FOnMessageBusShutdown);
@@ -26,7 +37,7 @@ DECLARE_MULTICAST_DELEGATE(FOnMessageBusShutdown);
  * Depending on their usage, messages are classified into a number of types, such as commands, events and documents. In Unreal Engine,
  * all these messages are modeled with regular built-in or user defined UStructs that may either be empty or contain data in the
  * form of UProperty fields[2]. Before being dispatched, messages are internally wrapped into so called Message Context objects
- * (see IMessageContext) that contain additional information about a message, such as when it was sent and the sender and recipients.
+ * (IMessageContext) that contain additional information about a message, such as when it was sent and the sender and recipients.
  *
  * The sending and receiving of messages is not limited to message endpoints within the same thread or process, but may be extended
  * to other applications running on the same computer or on other computers connected to a network with so called Message Transport
@@ -34,12 +45,12 @@ DECLARE_MULTICAST_DELEGATE(FOnMessageBusShutdown);
  * details of the underlying transport mechanisms, so that users can focus on implementing their distributed applications rather than
  * worrying about how data gets from one endpoint to another. A message bus makes it look as if all senders and recipients are located
  * within the same process, regardless of whether that is actually the case or not. It is possible to restrict the reach of messages
- * using so called Message Scopes (see EMessageScope).
+ * using so called Message Scopes (EMessageScope).
  *
- * All message recipients (objects implementing the see IReceiveMessages interface) must be registered with the message bus through the
+ * All message recipients (objects implementing the see IMessageReceiver interface) must be registered with the message bus through the
  * see IMessageBus.Register method. Before a recipient is destroyed it should unregister itself using the see IMessageBus.Unregister
- * method. Message senders (objects implementing the see ISendMessages interface) do not register with the bus, but instead pass a
- * reference to themselves each time they send a message. The IReceiveMessages and ISendMessages both are very low-level interfaces
+ * method. Message senders (objects implementing the see IMessageSender interface) do not register with the bus, but instead pass a
+ * reference to themselves each time they send a message. The IMessageReceiver and IMessageSender both are very low-level interfaces
  * into the messaging system. Most users will prefer to use instances of the see FMessageEndpoint class instead, which provides a much
  * more convenient way of sending and receiving messages.
  *
@@ -48,7 +59,7 @@ DECLARE_MULTICAST_DELEGATE(FOnMessageBusShutdown);
  * separate features in the future.
  *
  * In the Request-Reply pattern a message is sent to one or more particular message recipients using the IMessageBus.Send method.
- * Message recipients implement the see IMessageRecipient interface and are uniquely identified by their addresses (see FMessageAddress).
+ * Message recipients implement the see IMessageRecipient interface and are uniquely identified by their addresses (FMessageAddress).
  * After a message is received, the recipients may then reply with another message using the same IMessageBus.Send method. Alternatively,
  * a previously received message may be forwarded to another recipient using the see IMessageBus.Forward method. This pattern is useful
  * when message recipients already know about each other and wish to communicate directly, i.e. to exchange commands or events.
@@ -93,7 +104,7 @@ DECLARE_MULTICAST_DELEGATE(FOnMessageBusShutdown);
  *
  * [4] There are currently no built-in mechanisms for authentication and authorization, but they are on the to-do list.
  *
- * @see IReceiveMessages, ISendMessages, FMessageAddress, FMessageEndpoint
+ * @see FMessageAddress, FMessageEndpoint, IMessageReceiver, IMessageSender
  */
 class IMessageBus
 {
diff --git a/Engine/Source/Runtime/Messaging/Public/IMessageContext.h b/Engine/Source/Runtime/Messaging/Public/IMessageContext.h
--- a/Engine/Source/Runtime/Messaging/Public/IMessageContext.h
+++ b/Engine/Source/Runtime/Messaging/Public/IMessageContext.h
@@ -2,14 +2,23 @@
 
 #pragma once
 
-#include "CoreMinimal.h"
+#include "CoreTypes.h"
+#include "Async/TaskGraphInterfaces.h"
+#include "Containers/Array.h"
+#include "Containers/Map.h"
+#include "Misc/Crc.h"
 #include "Misc/Guid.h"
+#include "Templates/SharedPointer.h"
 #include "UObject/Class.h"
+#include "UObject/NameTypes.h"
 #include "UObject/WeakObjectPtr.h"
-#include "Async/TaskGraphInterfaces.h"
+#include "UObject/WeakObjectPtrTemplates.h"
 
 class IMessageAttachment;
 
+struct FDateTime;
+
+
 /**
  * Structure for message endpoint addresses.
  */
@@ -189,13 +198,13 @@ typedef TRangeBound<EMessageScope> FMessageScopeRangeBound;
  * message, such as its origin or when it expires. They may also be interested in optional out-of-band
  * binary data that is attached to the message.
  *
- * The sender's address (@see GetSender) is often needed to send a reply message to a message sender, i.e.
- * in response to a published message. The message attachment (@see GetAttchment) is an optional bundle
+ * The sender's address (IMessageContext.GetSender) is often needed to send a reply message to a message sender, i.e.
+ * in response to a published message. The message attachment (IMessageContext.GetAttachment) is an optional bundle
  * of binary bulk data that is transferred independently from the message itself and allows for transferring
  * larger amounts of data that would otherwise clog up the messaging system.
  *
  * In case a message was forwarded by another endpoint, the context of the original sender can be accessed
- * using the @see GetOriginalContext method.
+ * using the IMessageContext.GetOriginalContext method.
  *
  * @see IMessageAttachment
  */
@@ -342,7 +351,9 @@ public:
 
 
 /** Type definition for shared pointers to instances of IMessageContext. */
+DEPRECATED(4.16, "IMessageContextPtr is deprecated. Please use 'TSharedPtr<IMessageContext, ESPMode::ThreadSafe>' instead!")
 typedef TSharedPtr<IMessageContext, ESPMode::ThreadSafe> IMessageContextPtr;
 
 /** Type definition for shared references to instances of IMessageContext. */
+DEPRECATED(4.16, "IMessageContextRef is deprecated. Please use 'TSharedRef<IMessageContext, ESPMode::ThreadSafe>' instead!")
 typedef TSharedRef<IMessageContext, ESPMode::ThreadSafe> IMessageContextRef;
diff --git a/Engine/Source/Runtime/Messaging/Public/IMessageHandler.h b/Engine/Source/Runtime/Messaging/Public/IMessageHandler.h
--- a/Engine/Source/Runtime/Messaging/Public/IMessageHandler.h
+++ b/Engine/Source/Runtime/Messaging/Public/IMessageHandler.h
@@ -2,8 +2,10 @@
 
 #pragma once
 
-#include "CoreMinimal.h"
-#include "IMessageContext.h"
+#include "Templates/SharedPointer.h"
+
+class IMessageContext;
+
 
 /**
  * Interface for message handlers.
diff --git a/Engine/Source/Runtime/Messaging/Public/IMessageInterceptor.h b/Engine/Source/Runtime/Messaging/Public/IMessageInterceptor.h
--- a/Engine/Source/Runtime/Messaging/Public/IMessageInterceptor.h
+++ b/Engine/Source/Runtime/Messaging/Public/IMessageInterceptor.h
@@ -2,8 +2,12 @@
 
 #pragma once
 
-#include "CoreMinimal.h"
-#include "IMessageContext.h"
+#include "Templates/SharedPointer.h"
+#include "UObject/NameTypes.h"
+
+class IMessageContext;
+struct FGuid;
+
 
 /**
  * Interface for message interceptors.
diff --git a/Engine/Source/Runtime/Messaging/Public/IMessageReceiver.h b/Engine/Source/Runtime/Messaging/Public/IMessageReceiver.h
--- a/Engine/Source/Runtime/Messaging/Public/IMessageReceiver.h
+++ b/Engine/Source/Runtime/Messaging/Public/IMessageReceiver.h
@@ -2,20 +2,25 @@
 
 #pragma once
 
-#include "CoreMinimal.h"
-#include "IMessageContext.h"
+#include "Async/TaskGraphInterfaces.h"
+#include "Templates/SharedPointer.h"
+#include "UObject/NameTypes.h"
+
+class IMessageContext;
+struct FGuid;
+
 
 /**
  * Interface for message recipients.
  *
  * Classes that implement this interface are able to receive messages from a message bus. A message recipient will receive
- * a call to its ReceiveMessage() method for each message that was sent directly to it (see IMessageBus.Send) and for each
- * published message that it subscribed to (see IMessageBus.Publish).
+ * a call to its IMessageReceiver.ReceiveMessage method for each message that was sent directly to it (via IMessageBus.Send)
+ * and for each published message (via IMessageBus.Publish) that it subscribed to (via IMessageBus.Subscribe).
  *
  * This interface provides a rather low-level mechanism for receiving messages. Instead of implementing it, Most users will
  * want to use an instance of see FMessageEndpoint, which provides a much more convenient way of sending and receiving messages.
  *
- * @see FMessageEndpoint, IMessageBus, ISendMessages
+ * @see FMessageEndpoint, IMessageBus, IMessageSender
  */
 class IMessageReceiver
 {
@@ -63,7 +68,7 @@ public:
 	 *
 	 * @param Context Will hold the context of the received message.
 	 */
-	virtual void ReceiveMessage( const TSharedRef<IMessageContext, ESPMode::ThreadSafe>& Context ) = 0;
+	virtual void ReceiveMessage(const TSharedRef<IMessageContext, ESPMode::ThreadSafe>& Context) = 0;
 
 public:
 
diff --git a/Engine/Source/Runtime/Messaging/Public/IMessageSender.h b/Engine/Source/Runtime/Messaging/Public/IMessageSender.h
--- a/Engine/Source/Runtime/Messaging/Public/IMessageSender.h
+++ b/Engine/Source/Runtime/Messaging/Public/IMessageSender.h
@@ -2,26 +2,27 @@
 
 #pragma once
 
-#include "CoreMinimal.h"
 #include "IMessageContext.h"
+#include "Templates/SharedPointer.h"
+
+class FString;
 
-class Error;
 
 /**
  * Interface for message senders.
  *
  * Classes that implement this interface are able to send messages on a message bus. Each message sender must be uniquely
- * identifiable with a message address that is returned by the ISendMessages.GetSenderAddress method. It is recommended
+ * identifiable with a message address that is returned by the IMessageSender.GetSenderAddress method. It is recommended
  * that implementors of this interface generate a GUID for each instance, which must not change throughout the lifetime of
  * the instance.
  *
  * The sending of messages is accomplished with the IMessageBus.Forward, IMessageBus.Publish and IMessageBus.Send methods.
- * In case an error occurs during the sending of a message, the ISendMessages.NotifyMessageError method will be called.
+ * In case an error occurs during the sending of a message, the IMessageSender.NotifyMessageError method will be called.
  *
  * This interface provides a rather low-level mechanism for receiving messages. Instead of implementing it, Most users will
  * want to use an instance of FMessageEndpoint, which provides a much more convenient way of sending and receiving messages.
  *
- * @see FMessageEndpoint, IMessageBus, IReceiveMessages
+ * @see FMessageEndpoint, IMessageBus, IMessageReceiver
  */
 class IMessageSender
 {
@@ -40,7 +41,7 @@ public:
 	 * @param Context The context of the message that generated the error.
 	 * @param Error The error string.
 	 */
-	virtual void NotifyMessageError( const TSharedRef<IMessageContext, ESPMode::ThreadSafe>& Context, const FString& Error ) = 0;
+	virtual void NotifyMessageError(const TSharedRef<IMessageContext, ESPMode::ThreadSafe>& Context, const FString& Error) = 0;
 
 public:
 
diff --git a/Engine/Source/Runtime/Messaging/Public/IMessageSubscription.h b/Engine/Source/Runtime/Messaging/Public/IMessageSubscription.h
--- a/Engine/Source/Runtime/Messaging/Public/IMessageSubscription.h
+++ b/Engine/Source/Runtime/Messaging/Public/IMessageSubscription.h
@@ -2,7 +2,9 @@
 
 #pragma once
 
-#include "CoreMinimal.h"
+#include "Math/Range.h"
+#include "Templates/SharedPointer.h"
+#include "UObject/NameTypes.h"
 
 class IMessageReceiver;
 enum class EMessageScope : uint8;
diff --git a/Engine/Source/Runtime/Messaging/Public/IMessageTracer.h b/Engine/Source/Runtime/Messaging/Public/IMessageTracer.h
--- a/Engine/Source/Runtime/Messaging/Public/IMessageTracer.h
+++ b/Engine/Source/Runtime/Messaging/Public/IMessageTracer.h
@@ -2,13 +2,19 @@
 
 #pragma once
 
-#include "CoreMinimal.h"
+#include "CoreTypes.h"
+#include "Containers/Array.h"
+#include "Containers/Map.h"
+#include "Delegates/Delegate.h"
+#include "Templates/SharedPointer.h"
+
 #include "IMessageContext.h"
 
 struct FMessageTracerEndpointInfo;
 struct FMessageTracerMessageInfo;
 struct FMessageTracerTypeInfo;
 
+
 /**
  * Enumerates tracer breakpoint states.
  */
@@ -151,7 +157,7 @@ struct FMessageTracerInterceptorInfo
 struct FMessageTracerMessageInfo
 {
 	/** Holds a pointer to the message context. */
-	IMessageContextPtr Context;
+	TSharedPtr<IMessageContext, ESPMode::ThreadSafe> Context;
 
 	/** Holds the message's dispatch states per endpoint. */
 	TMap<TSharedPtr<FMessageTracerEndpointInfo>, TSharedPtr<FMessageTracerDispatchState>> DispatchStates;
diff --git a/Engine/Source/Runtime/Messaging/Public/IMessageTracerBreakpoint.h b/Engine/Source/Runtime/Messaging/Public/IMessageTracerBreakpoint.h
--- a/Engine/Source/Runtime/Messaging/Public/IMessageTracerBreakpoint.h
+++ b/Engine/Source/Runtime/Messaging/Public/IMessageTracerBreakpoint.h
@@ -2,8 +2,10 @@
 
 #pragma once
 
-#include "CoreMinimal.h"
-#include "IMessageContext.h"
+#include "Templates/SharedPointer.h"
+
+class IMessageContext;
+
 
 /**
  * Interface for message tracer breakpoints.
diff --git a/Engine/Source/Runtime/Messaging/Public/IMessageTransport.h b/Engine/Source/Runtime/Messaging/Public/IMessageTransport.h
--- a/Engine/Source/Runtime/Messaging/Public/IMessageTransport.h
+++ b/Engine/Source/Runtime/Messaging/Public/IMessageTransport.h
@@ -2,9 +2,14 @@
 
 #pragma once
 
-#include "CoreMinimal.h"
+#include "Containers/Array.h"
 #include "Misc/Guid.h"
-#include "IMessageContext.h"
+#include "Templates/SharedPointer.h"
+#include "UObject/NameTypes.h"
+
+class IMessageContext;
+class IMessageTransportHandler;
+
 
 /**
  * Interface for message transport technologies.
@@ -26,10 +31,11 @@ public:
 	/**
 	 * Starts up the message transport.
 	 *
+	 * @param Handler The handler of inbound transport messages and events.
 	 * @return Whether the transport was started successfully.
 	 * @see StopTransport
 	 */
-	virtual bool StartTransport() = 0;
+	virtual bool StartTransport(IMessageTransportHandler& Handler) = 0;
 
 	/**
 	 * Shuts down the message transport.
@@ -47,20 +53,6 @@ public:
 	 */
 	virtual bool TransportMessage(const TSharedRef<IMessageContext, ESPMode::ThreadSafe>& Context, const TArray<FGuid>& Recipients) = 0;
 
-public:
-
-	/** A delegate that is executed when message data has been received. */
-	DECLARE_DELEGATE_TwoParams(FOnMessageReceived, const IMessageContextRef& /*MessageContext*/, const FGuid& /*NodeId*/)
-	virtual FOnMessageReceived& OnMessageReceived() = 0;
-
-	/** A delegate that is executed when a transport node has been discovered. */
-	DECLARE_DELEGATE_OneParam(FOnNodeDiscovered, const FGuid& /*NodeId*/)
-	virtual FOnNodeDiscovered& OnNodeDiscovered() = 0;
-
-	/** A delegate that is executed when a transport node has closed or timed out. */
-	DECLARE_DELEGATE_OneParam(FOnNodeLost, const FGuid& /*NodeId*/)
-	virtual FOnNodeLost& OnNodeLost() = 0;
-
 protected:
 
 	/** Virtual destructor. */
@@ -69,4 +61,5 @@ protected:
 
 
 /** Type definition for shared references to instances of ITransportMessages. */
+DEPRECATED(4.16, "IMessageTransportRef is deprecated. Please use 'TSharedRef<IMessageTransport, ESPMode::ThreadSafe>' instead!")
 typedef TSharedRef<IMessageTransport, ESPMode::ThreadSafe> IMessageTransportRef;
diff --git a/Engine/Source/Runtime/Messaging/Public/IMessagingModule.h b/Engine/Source/Runtime/Messaging/Public/IMessagingModule.h
--- a/Engine/Source/Runtime/Messaging/Public/IMessagingModule.h
+++ b/Engine/Source/Runtime/Messaging/Public/IMessagingModule.h
@@ -2,14 +2,17 @@
 
 #pragma once
 
-#include "CoreMinimal.h"
 #include "Modules/ModuleInterface.h"
 #include "Modules/ModuleManager.h"
-#include "IMessageBus.h"
-#include "IMessageTransport.h"
+#include "Templates/SharedPointer.h"
 
 class IAuthorizeMessageRecipients;
 class IMessageBridge;
+class IMessageBus;
+class IMessageTransport;
+
+struct FMessageAddress;
+
 
 /**
  * Interface for messaging modules.
diff --git a/Engine/Source/Runtime/Messaging/Public/IMutableMessageContext.h b/Engine/Source/Runtime/Messaging/Public/IMutableMessageContext.h
--- a/Engine/Source/Runtime/Messaging/Public/IMutableMessageContext.h
+++ b/Engine/Source/Runtime/Messaging/Public/IMutableMessageContext.h
@@ -2,12 +2,19 @@
 
 #pragma once
 
-#include "CoreMinimal.h"
 #include "IMessageContext.h"
+#include "Templates/SharedPointer.h"
+
+class FName;
+class FString;
+class IMessageAttachment;
+class UScriptStruct;
 
 enum class EMessageScope : uint8;
+
+struct FDateTime;
 struct FMessageAddress;
-class IMessageAttachment;
+
 
 /**
  * Interface for mutable message contexts.
diff --git a/Engine/Source/Runtime/MessagingCommon/MessagingCommon.Build.cs b/Engine/Source/Runtime/MessagingCommon/MessagingCommon.Build.cs
new file mode 100644
--- /dev/null
+++ b/Engine/Source/Runtime/MessagingCommon/MessagingCommon.Build.cs
@@ -0,0 +1,28 @@
+// Copyright 1998-2017 Epic Games, Inc. All Rights Reserved.
+
+namespace UnrealBuildTool.Rules
+{
+	public class MessagingCommon : ModuleRules
+	{
+		public MessagingCommon(ReadOnlyTargetRules Target) : base(Target)
+		{
+			PublicDependencyModuleNames.AddRange(
+				new string[]
+				{
+					"Core",
+					"CoreUObject",
+				});
+
+			PublicIncludePathModuleNames.AddRange(
+				new string[]
+				{
+					"Messaging",
+				});
+
+			PrivateIncludePaths.AddRange(
+				new string[] {
+					"Runtime/MessagingCommon/Private",
+				});
+		}
+	}
+}
diff --git a/Engine/Source/Runtime/MessagingCommon/Private/MessagingCommonModule.cpp b/Engine/Source/Runtime/MessagingCommon/Private/MessagingCommonModule.cpp
new file mode 100644
--- /dev/null
+++ b/Engine/Source/Runtime/MessagingCommon/Private/MessagingCommonModule.cpp
@@ -0,0 +1,6 @@
+// Copyright 1998-2017 Epic Games, Inc. All Rights Reserved.
+
+#include "Modules/ModuleManager.h"
+
+
+IMPLEMENT_MODULE(FDefaultModuleImpl, MessagingCommon);
diff --git a/Engine/Source/Runtime/MessagingRpc/MessagingRpc.Build.cs b/Engine/Source/Runtime/MessagingRpc/MessagingRpc.Build.cs
--- a/Engine/Source/Runtime/MessagingRpc/MessagingRpc.Build.cs
+++ b/Engine/Source/Runtime/MessagingRpc/MessagingRpc.Build.cs
@@ -6,33 +6,30 @@ namespace UnrealBuildTool.Rules
 	{
 		public MessagingRpc(ReadOnlyTargetRules Target) : base(Target)
 		{
-            DynamicallyLoadedModuleNames.AddRange(
-                new string[]
-                {
-                    "Messaging",
-                }
-            );
+			DynamicallyLoadedModuleNames.AddRange(
+				new string[]
+				{
+					"Messaging",
+				});
 
 			PublicDependencyModuleNames.AddRange(
 				new string[]
 				{
 					"Core",
 					"CoreUObject",
-				}
-			);
+				});
 
-            PrivateIncludePathModuleNames.AddRange(
-                new string[]
-                {
-                    "Messaging",
-                }
-            );
+			PrivateIncludePathModuleNames.AddRange(
+				new string[]
+				{
+					"Messaging",
+					"MessagingCommon",
+				});
 
 			PrivateIncludePaths.AddRange(
 				new string[] {
 					"Runtime/MessagingRpc/Private",
-				}
-			);
+				});
 		}
 	}
 }
diff --git a/Engine/Source/Runtime/MessagingRpc/Private/MessageRpcClient.cpp b/Engine/Source/Runtime/MessagingRpc/Private/MessageRpcClient.cpp
--- a/Engine/Source/Runtime/MessagingRpc/Private/MessageRpcClient.cpp
+++ b/Engine/Source/Runtime/MessagingRpc/Private/MessageRpcClient.cpp
@@ -1,13 +1,16 @@
 // Copyright 1998-2017 Epic Games, Inc. All Rights Reserved.
 
 #include "MessageRpcClient.h"
+
 #include "Containers/ArrayBuilder.h"
-#include "Helpers/MessageEndpoint.h"
 #include "Containers/Ticker.h"
-#include "Helpers/MessageEndpointBuilder.h"
-#include "RpcMessage.h"
+#include "MessageEndpoint.h"
+#include "MessageEndpointBuilder.h"
+
+#include "IMessageRpcCall.h"
 #include "MessageRpcDefines.h"
 #include "MessageRpcMessages.h"
+#include "RpcMessage.h"
 
 
 /* FMessageRpcClient structors
diff --git a/Engine/Source/Runtime/MessagingRpc/Private/MessageRpcClient.h b/Engine/Source/Runtime/MessagingRpc/Private/MessageRpcClient.h
--- a/Engine/Source/Runtime/MessagingRpc/Private/MessageRpcClient.h
+++ b/Engine/Source/Runtime/MessagingRpc/Private/MessageRpcClient.h
@@ -2,15 +2,20 @@
 
 #pragma once
 
-#include "CoreMinimal.h"
-#include "Misc/Guid.h"
+#include "CoreTypes.h"
+#include "Containers/Map.h"
+#include "Delegates/IDelegateInstance.h"
 #include "IMessageContext.h"
-#include "IMessageRpcCall.h"
 #include "IMessageRpcClient.h"
+#include "Misc/Guid.h"
+#include "Templates/SharedPointer.h"
 
 class FMessageEndpoint;
+class IMessageRpcCall;
+
 struct FMessageRpcProgress;
 
+
 /**
  * Implements an RPC client.
  */
diff --git a/Engine/Source/Runtime/MessagingRpc/Private/MessageRpcDefines.h b/Engine/Source/Runtime/MessagingRpc/Private/MessageRpcDefines.h
--- a/Engine/Source/Runtime/MessagingRpc/Private/MessageRpcDefines.h
+++ b/Engine/Source/Runtime/MessagingRpc/Private/MessageRpcDefines.h
@@ -2,13 +2,9 @@
 
 #pragma once
 
-#include "CoreMinimal.h"
 
-/* Private constants
- *****************************************************************************/
-
-/** Defines interval at which calls are being re-sent to the server (in seconds). */
+/** The interval at which calls are being re-sent to the server (in seconds). */
 #define MESSAGE_RPC_RETRY_INTERVAL 1.0
 
-/** Defines the time after which calls time out (in seconds). */
+/** The time after which calls time out (in seconds). */
 #define MESSAGE_RPC_RETRY_TIMEOUT 3.0
diff --git a/Engine/Source/Runtime/MessagingRpc/Private/MessageRpcServer.cpp b/Engine/Source/Runtime/MessagingRpc/Private/MessageRpcServer.cpp
--- a/Engine/Source/Runtime/MessagingRpc/Private/MessageRpcServer.cpp
+++ b/Engine/Source/Runtime/MessagingRpc/Private/MessageRpcServer.cpp
@@ -1,15 +1,19 @@
 // Copyright 1998-2017 Epic Games, Inc. All Rights Reserved.
 
 #include "MessageRpcServer.h"
-#include "Containers/ArrayBuilder.h"
-#include "Helpers/MessageEndpoint.h"
+
 #include "Async/IAsyncProgress.h"
+#include "Containers/ArrayBuilder.h"
 #include "Containers/Ticker.h"
-#include "Helpers/MessageEndpointBuilder.h"
+#include "IMessageRpcHandler.h"
+#include "MessageEndpoint.h"
+#include "MessageEndpointBuilder.h"
+
 #include "RpcMessage.h"
 #include "MessageRpcDefines.h"
 #include "MessageRpcMessages.h"
 
+
 /* FMessageRpcServer structors
  *****************************************************************************/
 
diff --git a/Engine/Source/Runtime/MessagingRpc/Private/MessagingRpcModule.cpp b/Engine/Source/Runtime/MessagingRpc/Private/MessagingRpcModule.cpp
--- a/Engine/Source/Runtime/MessagingRpc/Private/MessagingRpcModule.cpp
+++ b/Engine/Source/Runtime/MessagingRpc/Private/MessagingRpcModule.cpp
@@ -1,12 +1,9 @@
 // Copyright 1998-2017 Epic Games, Inc. All Rights Reserved.
 
-#include "CoreMinimal.h"
+#include "IMessagingRpcModule.h"
 #include "Modules/ModuleManager.h"
-#include "IMessageRpcClient.h"
 #include "MessageRpcClient.h"
-#include "IMessageRpcServer.h"
 #include "MessageRpcServer.h"
-#include "IMessagingRpcModule.h"
 
 
 /**
@@ -17,7 +14,7 @@ class FMessagingRpcModule
 {
 public:
 
-	// IMessagingRpcModule interface
+	//~ IMessagingRpcModule interface
 
 	virtual TSharedRef<IMessageRpcClient> CreateRpcClient() override
 	{
@@ -31,21 +28,15 @@ public:
 
 public:
 
-	// IModuleInterface interface
-
-	virtual void StartupModule() override
-	{
-	}
+	//~ IModuleInterface interface
 
-	virtual void ShutdownModule() override
-	{
-	}
+	virtual void StartupModule() override { }
+	virtual void ShutdownModule() override { }
 
 	virtual bool SupportsDynamicReloading() override
 	{
 		return false;
 	}
-
 };
 
 
diff --git a/Engine/Source/Runtime/MessagingRpc/Public/IMessageRpcCall.h b/Engine/Source/Runtime/MessagingRpc/Public/IMessageRpcCall.h
--- a/Engine/Source/Runtime/MessagingRpc/Public/IMessageRpcCall.h
+++ b/Engine/Source/Runtime/MessagingRpc/Public/IMessageRpcCall.h
@@ -2,8 +2,14 @@
 
 #pragma once
 
-#include "CoreMinimal.h"
-#include "IMessageContext.h"
+#include "Misc/DateTime.h"
+
+class FText;
+class IMessageContext;
+class UScriptStruct;
+
+struct FGuid;
+
 
 /**
  * Interface for RPC calls.
diff --git a/Engine/Source/Runtime/MessagingRpc/Public/IMessageRpcHandler.h b/Engine/Source/Runtime/MessagingRpc/Public/IMessageRpcHandler.h
--- a/Engine/Source/Runtime/MessagingRpc/Public/IMessageRpcHandler.h
+++ b/Engine/Source/Runtime/MessagingRpc/Public/IMessageRpcHandler.h
@@ -2,11 +2,12 @@
 
 #pragma once
 
-#include "CoreMinimal.h"
-#include "IMessageContext.h"
+#include "Templates/SharedPointer.h"
 
+class IMessageContext;
 class IMessageRpcReturn;
 
+
 /**
  * Interface for RPC request handlers.
  */
diff --git a/Engine/Source/Runtime/MessagingRpc/Public/IMessageRpcReturn.h b/Engine/Source/Runtime/MessagingRpc/Public/IMessageRpcReturn.h
--- a/Engine/Source/Runtime/MessagingRpc/Public/IMessageRpcReturn.h
+++ b/Engine/Source/Runtime/MessagingRpc/Public/IMessageRpcReturn.h
@@ -2,10 +2,10 @@
 
 #pragma once
 
-#include "CoreMinimal.h"
-
+class UScriptStruct;
 struct FRpcMessage;
 
+
 /**
  * Interface for RPC return values.
  */
diff --git a/Engine/Source/Runtime/MessagingRpc/Public/IMessagingRpcModule.h b/Engine/Source/Runtime/MessagingRpc/Public/IMessagingRpcModule.h
--- a/Engine/Source/Runtime/MessagingRpc/Public/IMessagingRpcModule.h
+++ b/Engine/Source/Runtime/MessagingRpc/Public/IMessagingRpcModule.h
@@ -2,10 +2,12 @@
 
 #pragma once
 
-#include "CoreMinimal.h"
 #include "Modules/ModuleInterface.h"
-#include "IMessageRpcClient.h"
-#include "IMessageRpcServer.h"
+#include "Templates/SharedPointer.h"
+
+class IMessageRpcClient;
+class IMessageRpcServer;
+
 
 /**
  * Interface for the MessagingRpc module.
diff --git a/Engine/Source/Runtime/MessagingRpc/Public/MessageRpcMessages.h b/Engine/Source/Runtime/MessagingRpc/Public/MessageRpcMessages.h
--- a/Engine/Source/Runtime/MessagingRpc/Public/MessageRpcMessages.h
+++ b/Engine/Source/Runtime/MessagingRpc/Public/MessageRpcMessages.h
@@ -2,11 +2,13 @@
 
 #pragma once
 
-#include "CoreMinimal.h"
+#include "Containers/UnrealString.h"
 #include "UObject/ObjectMacros.h"
 #include "Misc/Guid.h"
+
 #include "MessageRpcMessages.generated.h"
 
+
 /** Message for canceling an RPC call. */
 USTRUCT()
 struct FMessageRpcCancel
diff --git a/Engine/Source/Runtime/MessagingRpc/Public/MessageRpcServer.h b/Engine/Source/Runtime/MessagingRpc/Public/MessageRpcServer.h
--- a/Engine/Source/Runtime/MessagingRpc/Public/MessageRpcServer.h
+++ b/Engine/Source/Runtime/MessagingRpc/Public/MessageRpcServer.h
@@ -1,18 +1,23 @@
 // Copyright 1998-2017 Epic Games, Inc. All Rights Reserved.
 #pragma once
 
-#include "CoreMinimal.h"
+#include "CoreTypes.h"
+#include "Containers/Map.h"
+#include "Misc/DateTime.h"
 #include "Misc/Guid.h"
 #include "IMessageContext.h"
-#include "IMessageRpcHandler.h"
 #include "IMessageRpcServer.h"
+#include "Templates/SharedPointer.h"
 
 class FMessageEndpoint;
 class FMessagingRpcModule;
 class IAsyncProgress;
 class IAsyncTask;
+class IMessageRpcHandler;
+
 struct FMessageRpcCancel;
 
+
 /**
  * Implements an RPC server.
  */
@@ -26,7 +31,8 @@ public:
 
 public:
 
-	// IMessageRpcServer interface
+	//~ IMessageRpcServer interface
+
 	virtual void AddHandler(const FName& RequestMessageType, const TSharedRef<IMessageRpcHandler>& Handler) override;
 	virtual const FMessageAddress& GetAddress() const override;
 	virtual FOnMessageRpcNoHandler& OnNoHandler() override;
diff --git a/Engine/Source/Runtime/MessagingRpc/Public/RpcMessage.h b/Engine/Source/Runtime/MessagingRpc/Public/RpcMessage.h
--- a/Engine/Source/Runtime/MessagingRpc/Public/RpcMessage.h
+++ b/Engine/Source/Runtime/MessagingRpc/Public/RpcMessage.h
@@ -2,9 +2,9 @@
 
 #pragma once
 
-#include "CoreMinimal.h"
-#include "UObject/ObjectMacros.h"
 #include "Misc/Guid.h"
+#include "UObject/ObjectMacros.h"
+
 #include "RpcMessage.generated.h"
 
 
diff --git a/Engine/Source/Runtime/MovieScene/Public/IMovieScenePlayer.h b/Engine/Source/Runtime/MovieScene/Public/IMovieScenePlayer.h
--- a/Engine/Source/Runtime/MovieScene/Public/IMovieScenePlayer.h
+++ b/Engine/Source/Runtime/MovieScene/Public/IMovieScenePlayer.h
@@ -110,6 +110,12 @@ public:
 	 */
 	virtual FMovieSceneSpawnRegister& GetSpawnRegister() { return NullRegister; }
 
+	/*
+	 * Called wehn an object is spawned by sequencer
+	 * 
+	 */
+	virtual void OnObjectSpawned(UObject* InObject, const FMovieSceneEvaluationOperand& Operand) {}
+
 	/**
 	 * Called whenever an object binding has been resolved to give the player a chance to interact with the objects before they are animated
 	 * 
