Removed `TUniquePtr` support for `TMediaObjectPool`:
It's not used as `TMediaObjectPool::AcquireUnique` and `TMediaObjectPool::ToUnique` don't compile. The problem is that `TUniquePtr` doesn't support externally supplied custom deleter instance, thus `MakeUnique` has no overload used in the methods mentioned above. The problem wasn't detected earlier because `TMediaObjectPool` is a template class, these methods weren't used and so weren't instantiated.

diff --git a/Engine/Source/Runtime/MediaUtils/Public/MediaObjectPool.h b/Engine/Source/Runtime/MediaUtils/Public/MediaObjectPool.h
--- a/Engine/Source/Runtime/MediaUtils/Public/MediaObjectPool.h
+++ b/Engine/Source/Runtime/MediaUtils/Public/MediaObjectPool.h
@@ -221,23 +221,6 @@ public:
 		return MakeShareable(Object, TDeleter(Storage));
 	}
 
-	/**
-	 * Acquire a tracked unique object from the pool.
-	 *
-	 * Unique objects do not need to be returned to the pool. They'll be
-	 * reclaimed automatically when the unique pointer is destroyed.
-	 *
-	 * @return The object.
-	 * @see Acquire, AcquireShared, Reset, ToUnique
-	 */
-	TUniquePtr<ObjectType> AcquireUnique()
-	{
-		ObjectType* Object = Acquire();
-		check(Object != nullptr);
-
-		return MakeUnique(Object, TDeleter(Storage));
-	}
-
 	/**
 	 * Get the number of objects available in the pool.
 	 *
@@ -260,18 +243,6 @@ public:
 		return MakeShareable(Object, TDeleter(Storage));
 	}
 
-	/**
-	 * Convert an object to a unique pooled object.
-	 *
-	 * @param Object The object to convert.
-	 * @return The unique object.
-	 * @see AcquireUnique, ToShared
-	 */
-	TUniquePtr<ObjectType> ToUnique(ObjectType* Object)
-	{
-		return MakeUnique(Object, TDeleter(Storage));
-	}
-
 	/**
 	 * Return the given object to the pool.
 	 *
