Fix Image Write for EXR

diff --git a/Engine/Source/Runtime/MediaUtils/Public/MediaRecorder.cpp b/Engine/Source/Runtime/MediaUtils/Public/MediaRecorder.cpp
--- a/Engine/Source/Runtime/MediaUtils/Public/MediaRecorder.cpp
+++ b/Engine/Source/Runtime/MediaUtils/Public/MediaRecorder.cpp
@@ -230,6 +230,12 @@ void FMediaRecorder::TickRecording()
 				}
 			}
 
+			bool bIsGammaCorrectionPreProcessingEnabled = false;
+			if (TargetImageFormat == EImageFormat::EXR && Sample->IsOutputSrgb())
+			{
+				bIsGammaCorrectionPreProcessingEnabled = true;
+			}
+
 			TUniquePtr<FImageWriteTask> ImageTask = MakeUnique<FImageWriteTask>();
 
 			// Set PixelData
@@ -248,7 +254,7 @@ void FMediaRecorder::TickRecording()
 				}
 
 				// Should we move the color buffer into a raw image data container.
-				bool bUseFMediaImagePixelData = bSetAlpha || (Sample->GetStride() != Size.X * NumChannels);
+				bool bUseFMediaImagePixelData = bSetAlpha || (Sample->GetStride() != Size.X * NumChannels) || bIsGammaCorrectionPreProcessingEnabled;
 
 				if (bUseFMediaImagePixelData)
 				{
@@ -288,6 +294,12 @@ void FMediaRecorder::TickRecording()
 				}
 			}
 
+			if (bIsGammaCorrectionPreProcessingEnabled)
+			{
+				const float DefaultGammaValue = 2.2f;
+				ImageTask->PixelPreProcessors.Add(TAsyncGammaCorrect<FColor>(DefaultGammaValue));
+			}
+
 			ImageTask->Format = TargetImageFormat;
 			ImageTask->CompressionQuality = CompressionQuality;
 			ImageTask->bOverwriteFile = true;
