Moved AppMediaTimeSource to different log channel 'LogMediaTimeSource'.

diff --git a/Engine/Source/Runtime/MediaUtils/Private/AppMediaTimeSource.cpp b/Engine/Source/Runtime/MediaUtils/Private/AppMediaTimeSource.cpp
--- a/Engine/Source/Runtime/MediaUtils/Private/AppMediaTimeSource.cpp
+++ b/Engine/Source/Runtime/MediaUtils/Private/AppMediaTimeSource.cpp
@@ -6,6 +6,8 @@
 #include "Misc/App.h"
 #include "Misc/Timespan.h"
 
+/** Declares a log category for this module. */
+DECLARE_LOG_CATEGORY_CLASS(LogMediaTimeSource, Log, All);
 
 /* IMediaTimeSource interface 
  *****************************************************************************/
@@ -15,7 +17,7 @@ FTimespan FAppMediaTimeSource::GetTimecode()
 	const double CurrentTime = FApp::GetCurrentTime();
 	const FTimespan Timecode = FTimespan::FromSeconds(CurrentTime);
 
-	UE_LOG(LogMediaUtils, VeryVerbose, TEXT("AppMediaTimeSource: Time %.10f, Delta %.10f, Timecode %s"), CurrentTime, FApp::GetDeltaTime(), *Timecode.ToString(TEXT("%h:%m:%s.%t")));
+	UE_LOG(LogMediaTimeSource, VeryVerbose, TEXT("AppMediaTimeSource: Time %.10f, Delta %.10f, Timecode %s"), CurrentTime, FApp::GetDeltaTime(), *Timecode.ToString(TEXT("%h:%m:%s.%t")));
 
 	return Timecode;
 }
