mediaplayer: integration of updates for 10.40 release (output buffer handling for Switch and XB1 reworked; Electra player update; minor PS4 player updates reg. looping)
#rnx

diff --git a/Engine/Plugins/Media/WmfMedia/Source/WmfMedia/WmfMedia.Build.cs b/Engine/Plugins/Media/WmfMedia/Source/WmfMedia/WmfMedia.Build.cs
--- a/Engine/Plugins/Media/WmfMedia/Source/WmfMedia/WmfMedia.Build.cs
+++ b/Engine/Plugins/Media/WmfMedia/Source/WmfMedia/WmfMedia.Build.cs
@@ -18,25 +18,39 @@ namespace UnrealBuildTool.Rules
                     "WmfMediaFactory"
                     });
 
-            PrivateDependencyModuleNames.AddRange(
+			PrivateDependencyModuleNames.AddRange(
 				new string[] {
 					"Core",
 					"CoreUObject",
                     "Engine",
                     "MediaUtils",
-                    "Projects",
-                    "RenderCore",
+					"Projects",
+					"RenderCore",
                     "RHI",
-                    "UtilityShaders",
+					"UtilityShaders",
                     "WmfMediaFactory",
                 });
 
-            if (Target.Platform != UnrealTargetPlatform.XboxOne)
+            if (Target.Platform == UnrealTargetPlatform.XboxOne)
             {
-                PrivateDependencyModuleNames.Add("D3D11RHI");
+                PrivateDependencyModuleNames.AddRange(
+				  new string[] {
+					"D3D12RHI",
+                });
+                PrivateIncludePaths.AddRange(
+                    new string[] {
+                    "../../../../Source/Runtime/D3D12RHI/Private",
+                });
+            }
+            else
+            {
+                PrivateDependencyModuleNames.AddRange(
+                  new string[] {
+                    "D3D11RHI",
+                });
             }
 
-			PrivateIncludePathModuleNames.AddRange(
+            PrivateIncludePathModuleNames.AddRange(
 				new string[] {
 					"Media",
 				});
@@ -67,7 +81,7 @@ namespace UnrealBuildTool.Rules
 				PublicDelayLoadDLLs.Add("mfplat.dll");
 				PublicDelayLoadDLLs.Add("mfplay.dll");
 				PublicDelayLoadDLLs.Add("shlwapi.dll");
-            }
+			}
 		}
 	}
 }
diff --git a/Engine/Source/Runtime/Media/Public/IMediaTextureSampleConverter.h b/Engine/Source/Runtime/Media/Public/IMediaTextureSampleConverter.h
--- a/Engine/Source/Runtime/Media/Public/IMediaTextureSampleConverter.h
+++ b/Engine/Source/Runtime/Media/Public/IMediaTextureSampleConverter.h
@@ -2,6 +2,7 @@
 
 #pragma once
 
+#include "RHI.h"
 #include "RHIResources.h"
 
 /**
diff --git a/Engine/Source/Runtime/MediaAssets/MediaAssets.Build.cs b/Engine/Source/Runtime/MediaAssets/MediaAssets.Build.cs
--- a/Engine/Source/Runtime/MediaAssets/MediaAssets.Build.cs
+++ b/Engine/Source/Runtime/MediaAssets/MediaAssets.Build.cs
@@ -8,7 +8,6 @@ namespace UnrealBuildTool.Rules
 		{
 			PublicDependencyModuleNames.AddRange(
 				new string[] {
-					"SignalProcessing",
 					"AudioMixer",
 					"Core",
 					"CoreUObject",
diff --git a/Engine/Source/Runtime/MediaAssets/Private/Misc/MediaTextureResource.cpp b/Engine/Source/Runtime/MediaAssets/Private/Misc/MediaTextureResource.cpp
--- a/Engine/Source/Runtime/MediaAssets/Private/Misc/MediaTextureResource.cpp
+++ b/Engine/Source/Runtime/MediaAssets/Private/Misc/MediaTextureResource.cpp
@@ -32,14 +32,6 @@ DECLARE_CYCLE_STAT(TEXT("MediaAssets MediaTextureResource Render"), STAT_MediaAs
 /** Sample time of texture last rendered. */
 DECLARE_FLOAT_COUNTER_STAT(TEXT("MediaAssets MediaTextureResource Sample"), STAT_MediaUtils_TextureSampleTime, STATGROUP_Media);
 
-
-static int32 CachedSamplesQueueDepth = 1;
-static FAutoConsoleVariableRef CVarEncoderSaveVideoToFile(
-	TEXT("media.CachedSamplesQueueDepth"),
-	CachedSamplesQueueDepth,
-	TEXT("How many frames to hold samples before release (default = 1)."),
-	ECVF_Default);
-
 /* Local helpers
  *****************************************************************************/
 
@@ -174,11 +166,6 @@ FMediaTextureResource::FMediaTextureResource(UMediaTexture& InOwner, FIntPoint&
 	, OwnerDim(InOwnerDim)
 	, OwnerSize(InOwnerSize)
 {
-	// preset the CachedSamples. makes things easier in the long run
-	if (CachedSamplesQueueDepth > 0)
-	{
-		CachedSamples.AddDefaulted(CachedSamplesQueueDepth);
-	}
 }
 
 
@@ -189,8 +176,6 @@ void FMediaTextureResource::Render(const FRenderParams& Params)
 {
 	check(IsInRenderingThread());
 
-	CycleCachedSamples();
-
 	SCOPE_CYCLE_COUNTER(STAT_MediaAssets_MediaTextureResourceRender);
 
 	FLinearColor Rotation(1, 0, 0, 1);
@@ -302,11 +287,27 @@ void FMediaTextureResource::Render(const FRenderParams& Params)
 		}
 #endif
 
-		// We're not done with `Sample` as rendering is asynchronous. 
-		// Hold a reference in a member to postpone recycling `Sample` until safe to release
-		if (CachedSamplesQueueDepth > 0)
+		/*
+			We need to keep any sample used for display around to avoid any player reusing it before any async tasks are done.
+
+			Two cases are to be considered:
+			a) we directly use the textures in the sample -> we need to keep it as long as the sample is visible
+			b) we copy (e.g. as part of a conversion) the sample into a local texture -> we only need it for the current rendering interval
+
+			In both cases the FTextureSampleKeeper will make sure that it is kept for the current rendering interval per FRHIResource rules
+			for delayed deletion.
+		*/
+		TRefCountPtr<FTextureSampleKeeper> NewCurrentSample = TRefCountPtr<FTextureSampleKeeper>(new FTextureSampleKeeper(Sample));
+		check(NewCurrentSample);
+		if (RenderTargetTextureRHI != OutputTarget)
 		{
-			CachedSamples[0] = Sample;
+			// We use the texture data in the sample
+			CurrentSample = NewCurrentSample;
+		}
+		else
+		{
+			// We use the internal copy
+			CurrentSample = nullptr;
 		}
 	}
 	else if (Params.CanClear)
@@ -321,6 +322,9 @@ void FMediaTextureResource::Render(const FRenderParams& Params)
 #endif
 
 			ClearTexture(Params.ClearColor, Params.SrgbOutput);
+
+			// Also get rid of any sample from previous rendering...
+			CurrentSample = nullptr;
 		}
 	}
 	
@@ -346,16 +350,6 @@ void FMediaTextureResource::Render(const FRenderParams& Params)
 	Owner.SetRenderedExternalTextureGuid(Params.CurrentGuid);
 }
 
-void FMediaTextureResource::CycleCachedSamples()
-{
-	int32 CachedSampleToChange = CachedSamplesQueueDepth - 1;
-
-	while (CachedSampleToChange > 0)
-	{
-		CachedSamples[CachedSampleToChange] = CachedSamples[CachedSampleToChange - 1];
-		CachedSampleToChange--;
-	}
-}
 
 /* FRenderTarget interface
  *****************************************************************************/
diff --git a/Engine/Source/Runtime/MediaAssets/Private/Misc/MediaTextureResource.h b/Engine/Source/Runtime/MediaAssets/Private/Misc/MediaTextureResource.h
--- a/Engine/Source/Runtime/MediaAssets/Private/Misc/MediaTextureResource.h
+++ b/Engine/Source/Runtime/MediaAssets/Private/Misc/MediaTextureResource.h
@@ -158,10 +158,18 @@ protected:
 
 private:
 
-	/**
-	* Cycle our cached samples
-	*/
-	void CycleCachedSamples();
+	// Class to maintain a reference to a IMediaTextureSample instance as long  as needed by RHI
+	// (needed not to keep Texture for GPU - that is safe already - but to avoid reusing the buffer too early)
+	class FTextureSampleKeeper : public FRHIResource
+	{
+	public:
+		FTextureSampleKeeper(const TSharedPtr<IMediaTextureSample, ESPMode::ThreadSafe> & InMediaSample)
+			: MediaSample(InMediaSample)
+		{}
+
+		TSharedPtr<IMediaTextureSample, ESPMode::ThreadSafe> MediaSample;
+	};
+
 
 	/** Whether the texture has been cleared. */
 	bool Cleared;
@@ -191,5 +199,5 @@ private:
 	TWeakPtr<FMediaPlayerFacade, ESPMode::ThreadSafe> PlayerFacadePtr;
 
 	/** cached media sample to postpone releasing it until the next sample rendering as it can get overwritten due to asynchronous rendering */
-	TArray<TSharedPtr<IMediaTextureSample, ESPMode::ThreadSafe>> CachedSamples;
+	TRefCountPtr<FTextureSampleKeeper> CurrentSample;
 };
diff --git a/Engine/Source/Runtime/MediaUtils/Public/MediaObjectPool.h b/Engine/Source/Runtime/MediaUtils/Public/MediaObjectPool.h
--- a/Engine/Source/Runtime/MediaUtils/Public/MediaObjectPool.h
+++ b/Engine/Source/Runtime/MediaUtils/Public/MediaObjectPool.h
@@ -35,6 +35,14 @@ public:
 	 */
 	virtual void ShutdownPoolable() { }
 
+	/**
+	 * Used to check if returned object is ready for reuse right away
+	 */
+	virtual bool IsReadyForReuse()
+	{
+		return true;
+	}
+
 public:
 
 	/** Virtual destructor. */
@@ -42,6 +50,12 @@ public:
 };
 
 
+template<typename ObjectType> class TMediaPoolDefaultObjectAllocator
+{
+public:
+	static ObjectType *Alloc() { return new ObjectType; }
+};
+
 /**
  * Template for media object pools.
  *
@@ -50,7 +64,7 @@ public:
  * @param ObjectType The type of objects managed by this pool. 
  * @todo gmp: Make media object pool lock-free
  */
-template<typename ObjectType>
+template<typename ObjectType, typename ObjectAllocator=TMediaPoolDefaultObjectAllocator<ObjectType>>
 class TMediaObjectPool
 {
 	static_assert(TPointerIsConvertibleFromTo<ObjectType, IMediaPoolable>::Value, "Poolable objects must implement the IMediaPoolable interface.");
@@ -59,11 +73,20 @@ class TMediaObjectPool
 	class TStorage
 	{
 	public:
+		TStorage(ObjectAllocator *InObjectAllocatorInstance)
+			: ObjectAllocatorInstance(InObjectAllocatorInstance)
+		{}
 
 		/** Destructor. */
 		~TStorage()
 		{
 			Reserve(0);
+			ObjectType* Object;
+			while (WaitReadyForReuse.Dequeue(Object))
+			{
+				Object->ShutdownPoolable();
+				delete Object;
+			}
 		}
 
 	public:
@@ -79,11 +102,26 @@ class TMediaObjectPool
 				{
 					Result = Pool.Pop(false);
 				}
+				else
+				{
+					if (WaitReadyForReuse.Peek(Result))
+					{
+						if (Result->IsReadyForReuse())
+						{
+							WaitReadyForReuse.Pop();
+							Result->ShutdownPoolable();
+						}
+						else
+						{
+							Result = nullptr;
+						}
+					}
+				}
 			}
 
 			if (Result == nullptr)
 			{
-				Result = new ObjectType;
+				Result = ObjectAllocatorInstance->Alloc();
 			}
 			
 			Result->InitializePoolable();
@@ -108,8 +146,15 @@ class TMediaObjectPool
 
 			FScopeLock Lock(&CriticalSection);
 
-			Object->ShutdownPoolable();
-			Pool.Push(Object);
+			if (Object->IsReadyForReuse())
+			{
+				Object->ShutdownPoolable();
+				Pool.Push(Object);
+			}
+			else
+			{
+				WaitReadyForReuse.Enqueue(Object);
+			}
 		}
 
 		/** Reserve the specified number of objects. */
@@ -124,7 +169,32 @@ class TMediaObjectPool
 
 			while (NumObjects > (uint32)Pool.Num())
 			{
-				Pool.Push(new ObjectType);
+				Pool.Push(ObjectAllocatorInstance->Alloc());
+			}
+		}
+
+		/** Regular tick call */
+		void Tick()
+		{
+			if (WaitReadyForReuse.IsEmpty())
+			{
+				// Conservative early out to avoid CS: we will get any items missed the next time around
+				return;
+			}
+
+			FScopeLock Lock(&CriticalSection);
+
+			ObjectType* Object;
+			while (WaitReadyForReuse.Peek(Object))
+			{
+				if (!Object->IsReadyForReuse())
+				{
+					break;
+				}
+
+				Object->ShutdownPoolable();
+				Pool.Push(Object);
+				WaitReadyForReuse.Pop();
 			}
 		}
 
@@ -135,6 +205,12 @@ class TMediaObjectPool
 
 		/** List of unused objects. */
 		TArray<ObjectType*> Pool;
+
+		/** List of unused objects, waiting for reuse-ability. */
+		TQueue<ObjectType*> WaitReadyForReuse;
+
+		/** Object allocator instance (nullptr by default) */
+		ObjectAllocator *ObjectAllocatorInstance;
 	};
 
 
@@ -172,8 +248,8 @@ class TMediaObjectPool
 public:
 
 	/** Default constructor. */
-	TMediaObjectPool()
-		: Storage(MakeShareable(new TStorage))
+	TMediaObjectPool(ObjectAllocator *ObjectAllocatorInstance = nullptr)
+		: Storage(MakeShareable(new TStorage(ObjectAllocatorInstance)))
 	{ }
 
 	/**
@@ -268,6 +344,14 @@ public:
 		Storage->Reserve(NumObjects);
 	}
 
+	/**
+	 * Regular tick call
+	 */
+	void Tick()
+	{
+		Storage->Tick();
+	}
+
 private:
 
 	/** Storage for pooled objects. */
