#UE4 Nonunity fix

diff --git a/Engine/Source/Runtime/MediaUtils/Public/MediaObjectPool.h b/Engine/Source/Runtime/MediaUtils/Public/MediaObjectPool.h
--- a/Engine/Source/Runtime/MediaUtils/Public/MediaObjectPool.h
+++ b/Engine/Source/Runtime/MediaUtils/Public/MediaObjectPool.h
@@ -3,6 +3,7 @@
 #pragma once
 
 #include "Containers/Array.h"
+#include "Containers/Queue.h"
 #include "Misc/ScopeLock.h"
 #include "Templates/SharedPointer.h"
 #include "Templates/UniquePtr.h"
