MediaFramework now tracks audio/video decode time values and audio playback time values.
Used primarily for automated tests.
Merged from CL#4437471.

diff --git a/Engine/Source/Runtime/MediaAssets/Private/Assets/MediaSoundComponent.cpp b/Engine/Source/Runtime/MediaAssets/Private/Assets/MediaSoundComponent.cpp
--- a/Engine/Source/Runtime/MediaAssets/Private/Assets/MediaSoundComponent.cpp
+++ b/Engine/Source/Runtime/MediaAssets/Private/Assets/MediaSoundComponent.cpp
@@ -29,7 +29,8 @@ FAutoConsoleVariableRef CVarSyncAudioAfterDropouts(
 	TEXT("0: Not Enabled, 1: Enabled"),
 	ECVF_Default);
 
-DECLARE_FLOAT_COUNTER_STAT(TEXT("MediaUtils MediaSoundComponent Sync"), STAT_MediaUtils_MediaSoundComponent, STATGROUP_Media);
+DECLARE_FLOAT_COUNTER_STAT(TEXT("MediaUtils MediaSoundComponent Sync"), STAT_MediaUtils_MediaSoundComponentSync, STATGROUP_Media);
+DECLARE_FLOAT_COUNTER_STAT(TEXT("MediaUtils MediaSoundComponent SampleTime"), STAT_MediaUtils_MediaSoundComponentSampleTime, STATGROUP_Media);
 
 /* Static initialization
  *****************************************************************************/
@@ -52,6 +53,7 @@ UMediaSoundComponent::UMediaSoundComponent(const FObjectInitializer& ObjectIniti
 	, Resampler(new FMediaAudioResampler)
 	, FrameSyncOffset(0)
 	, bSyncAudioAfterDropouts(false)
+	, LastPlaySampleTime(FTimespan::MinValue())
 {
 	PrimaryComponentTick.bCanEverTick = true;
 	bAutoActivate = true;
@@ -151,6 +153,8 @@ void UMediaSoundComponent::UpdatePlayer()
 	// caching play rate and time for audio thread (eventual consistency is sufficient)
 	CachedRate = PlayerFacade->GetRate();
 	CachedTime = PlayerFacade->GetTime();
+
+	PlayerFacade->SetLastAudioRenderedSampleTime(LastPlaySampleTime.Load());
 }
 
 
@@ -404,9 +408,10 @@ int32 UMediaSoundComponent::OnGenerateAudio(float* OutAudio, int32 NumSamples)
 			}
 		}
 
-#if STATS
-		SET_FLOAT_STAT(STAT_MediaUtils_MediaSoundComponent, (Time - OutTime).GetTotalMilliseconds());
-#endif
+		LastPlaySampleTime = OutTime;
+
+		SET_FLOAT_STAT(STAT_MediaUtils_MediaSoundComponentSync, FMath::Abs((Time - OutTime).GetTotalMilliseconds()));
+		SET_FLOAT_STAT(STAT_MediaUtils_MediaSoundComponentSampleTime, OutTime.GetTotalMilliseconds());
 	}
 	else
 	{
@@ -417,6 +422,8 @@ int32 UMediaSoundComponent::OnGenerateAudio(float* OutAudio, int32 NumSamples)
 			FScopeLock Lock(&CriticalSection);
 			FrameSyncOffset = 0;
 		}
+
+		LastPlaySampleTime = FTimespan::MinValue();
 	}
 	return NumSamples;
 }
diff --git a/Engine/Source/Runtime/MediaAssets/Public/MediaSoundComponent.h b/Engine/Source/Runtime/MediaAssets/Public/MediaSoundComponent.h
--- a/Engine/Source/Runtime/MediaAssets/Public/MediaSoundComponent.h
+++ b/Engine/Source/Runtime/MediaAssets/Public/MediaSoundComponent.h
@@ -227,6 +227,9 @@ private:
 	 */
 	bool bSyncAudioAfterDropouts;
 
+	/* Time of last sample played. */
+	TAtomic<FTimespan> LastPlaySampleTime;
+
 private:
 
 	static USoundClass* DefaultMediaSoundClassObject;
diff --git a/Engine/Source/Runtime/MediaUtils/Private/MediaPlayerFacade.cpp b/Engine/Source/Runtime/MediaUtils/Private/MediaPlayerFacade.cpp
--- a/Engine/Source/Runtime/MediaUtils/Private/MediaPlayerFacade.cpp
+++ b/Engine/Source/Runtime/MediaUtils/Private/MediaPlayerFacade.cpp
@@ -12,6 +12,7 @@
 #include "IMediaPlayer.h"
 #include "IMediaPlayerFactory.h"
 #include "IMediaSamples.h"
+#include "IMediaAudioSample.h"
 #include "IMediaTextureSample.h"
 #include "IMediaTracks.h"
 #include "IMediaView.h"
@@ -72,7 +73,7 @@ FMediaPlayerFacade::FMediaPlayerFacade()
 	, BlockOnTime(FTimespan::MinValue())
 	, Cache(new FMediaSampleCache)
 	, LastRate(0.0f)
-	, NextVideoSampleTime(FTimespan::MinValue())
+	, LastAudioRenderedSampleTime(FTimespan::MinValue())
 { }
 
 
@@ -616,6 +617,25 @@ bool FMediaPlayerFacade::SupportsRate(float Rate, bool Unthinned) const
 	return Player.IsValid() && Player->GetControls().GetSupportedRates(Thinning).Contains(Rate);
 }
 
+void FMediaPlayerFacade::SetLastAudioRenderedSampleTime(FTimespan SampleTime)
+{
+	LastAudioRenderedSampleTime = SampleTime;
+}
+
+FTimespan FMediaPlayerFacade::GetLastAudioRenderedSampleTime() const
+{
+	return LastAudioRenderedSampleTime.Load();
+}
+
+FTimespan FMediaPlayerFacade::GetLastAudioSampleProcessedTime() const
+{
+	return LastAudioSampleProcessedTime.Load();
+}
+
+FTimespan FMediaPlayerFacade::GetLastVideoSampleProcessedTime() const
+{
+	return LastVideoSampleProcessedTime.Load();
+}
 
 /* FMediaPlayerFacade implementation
 *****************************************************************************/
@@ -677,6 +697,10 @@ void FMediaPlayerFacade::FlushSinks()
 	}
 
 	NextVideoSampleTime = FTimespan::MinValue();
+
+	LastAudioSampleProcessedTime = FTimespan::MinValue();
+	LastVideoSampleProcessedTime = FTimespan::MinValue();
+	LastAudioRenderedSampleTime = FTimespan::MinValue();
 }
 
 
@@ -1071,7 +1095,14 @@ void FMediaPlayerFacade::ProcessAudioSamples(IMediaSamples& Samples, TRange<FTim
 
 	while (Samples.FetchAudio(TimeRange, Sample))
 	{
-		if (Sample.IsValid() && !AudioSampleSinks.Enqueue(Sample.ToSharedRef(), FMediaPlayerQueueDepths::MaxAudioSinkDepth))
+		if (!Sample.IsValid())
+		{
+			continue;
+		}
+
+		LastAudioSampleProcessedTime = Sample->GetTime();
+
+		if (!AudioSampleSinks.Enqueue(Sample.ToSharedRef(), FMediaPlayerQueueDepths::MaxAudioSinkDepth))
 		{
 			#if MEDIAPLAYERFACADE_TRACE_SINKOVERFLOWS
 				UE_LOG(LogMediaUtils, VeryVerbose, TEXT("PlayerFacade %p: Audio sample sink overflow"), this);
@@ -1140,6 +1171,8 @@ void FMediaPlayerFacade::ProcessVideoSamples(IMediaSamples& Samples, TRange<FTim
 			continue;
 		}
 
+		LastVideoSampleProcessedTime = Sample->GetTime();
+
 		UE_LOG(LogMediaUtils, VeryVerbose, TEXT("PlayerFacade %p: Fetched video sample %s"), this, *Sample->GetTime().ToString(TEXT("%h:%m:%s.%t")));
 
 		if (VideoSampleSinks.Enqueue(Sample.ToSharedRef(), FMediaPlayerQueueDepths::MaxVideoSinkDepth))
diff --git a/Engine/Source/Runtime/MediaUtils/Public/MediaPlayerFacade.h b/Engine/Source/Runtime/MediaUtils/Public/MediaPlayerFacade.h
--- a/Engine/Source/Runtime/MediaUtils/Public/MediaPlayerFacade.h
+++ b/Engine/Source/Runtime/MediaUtils/Public/MediaPlayerFacade.h
@@ -631,6 +631,35 @@ public:
 	 */
 	bool SupportsRate(float Rate, bool Unthinned) const;
 
+	/**
+	 * Record last audio sample played to track audio sync (for automated tests)
+	 *
+	 * @param SampleTime Time of media sample currently being played
+	 * @return true if playback is being prepared, false otherwise.
+	 */
+	void SetLastAudioRenderedSampleTime(FTimespan SampleTime);
+
+	/**
+	 * Get time of last audio sample played
+	 *
+	 * @return Time of last audio sample played.
+	 */
+	FTimespan GetLastAudioRenderedSampleTime() const;
+
+	/**
+	 * Get time of last audio sample decoded
+	 *
+	 * @return Time of last audio sample decoded.
+	 */
+	FTimespan GetLastAudioSampleProcessedTime() const;
+
+	/**
+	 * Get time of last video sample decoded
+	 *
+	 * @return Time of last video sample decoded.
+	 */
+	FTimespan GetLastVideoSampleProcessedTime() const;
+
 public:
 
 	/** Get an event delegate that is invoked when a media event occurred. */
@@ -781,4 +810,13 @@ private:
 
 	/** Active media player options. */
 	TOptional<FMediaPlayerOptions> ActivePlayerOptions;
+
+	/** Time of last audio sample played. */
+	TAtomic<FTimespan> LastAudioRenderedSampleTime;
+
+	/** Time of last audio sample decoded. */
+	TAtomic<FTimespan> LastAudioSampleProcessedTime;
+
+	/** Time of last video sample decoded. */
+	TAtomic<FTimespan> LastVideoSampleProcessedTime;
 };
