Introducing a tighter active period for FMediaTicker thread & new optional feature to destroy players after closing

diff --git a/Engine/Source/Runtime/Media/Public/IMediaPlayer.h b/Engine/Source/Runtime/Media/Public/IMediaPlayer.h
--- a/Engine/Source/Runtime/Media/Public/IMediaPlayer.h
+++ b/Engine/Source/Runtime/Media/Public/IMediaPlayer.h
@@ -321,6 +321,16 @@ public:
 		// Override in child class if needed.
 	}
 
+	enum class FeatureFlag {
+		AllowShutdownOnClose = 0,	//!< Allow player to be shutdown right after 'close' event is received from it
+	};
+	
+	virtual bool GetPlayerFeatureFlag(FeatureFlag /*flag*/) const
+	{
+		// Override in child class if needed.
+		return false;
+	}
+
 public:
 
 	/** Virtual destructor. */
diff --git a/Engine/Source/Runtime/MediaAssets/Private/Assets/MediaPlayer.cpp b/Engine/Source/Runtime/MediaAssets/Private/Assets/MediaPlayer.cpp
--- a/Engine/Source/Runtime/MediaAssets/Private/Assets/MediaPlayer.cpp
+++ b/Engine/Source/Runtime/MediaAssets/Private/Assets/MediaPlayer.cpp
@@ -53,7 +53,7 @@ UMediaPlayer::UMediaPlayer(const FObjectInitializer& ObjectInitializer)
 {
 	if (!HasAnyFlags(RF_ClassDefaultObject))
 	{
-		PlayerFacade = MakeShared<FMediaPlayerFacade, ESPMode::ThreadSafe>();
+		PlayerFacade = MakeShareable(new FMediaPlayerFacade());
 		PlayerFacade->OnMediaEvent().AddUObject(this, &UMediaPlayer::HandlePlayerMediaEvent);
 		Playlist = NewObject<UMediaPlaylist>(GetTransientPackage(), NAME_None, RF_Transactional | RF_Transient);
 	}
@@ -804,8 +804,8 @@ void UMediaPlayer::RegisterWithMediaModule()
 
 	if (MediaModule != nullptr)
 	{
+		// Make sure the PlayerFacade instance gets regular tick calls from various spots in the gameloop
 		MediaModule->GetClock().AddSink(PlayerFacade.ToSharedRef());
-		MediaModule->GetTicker().AddTickable(PlayerFacade.ToSharedRef());
 		RegisteredWithMediaModule = true;
 	}
 	else
diff --git a/Engine/Source/Runtime/MediaUtils/Private/MediaPlayerFacade.cpp b/Engine/Source/Runtime/MediaUtils/Private/MediaPlayerFacade.cpp
--- a/Engine/Source/Runtime/MediaUtils/Private/MediaPlayerFacade.cpp
+++ b/Engine/Source/Runtime/MediaUtils/Private/MediaPlayerFacade.cpp
@@ -16,6 +16,7 @@
 #include "IMediaTextureSample.h"
 #include "IMediaTracks.h"
 #include "IMediaView.h"
+#include "IMediaTicker.h"
 #include "MediaPlayerOptions.h"
 #include "Math/NumericLimits.h"
 #include "Misc/CoreMisc.h"
@@ -448,6 +449,11 @@ bool FMediaPlayerFacade::Open(const FString& Url, const IMediaOptions* Options,
 		return false;
 	}
 
+	static const FName MediaModuleName("Media");
+	IMediaModule* MediaModule;
+	MediaModule = FModuleManager::GetModulePtr<IMediaModule>(MediaModuleName);
+	check(MediaModule);
+
 	// find & initialize new player
 	TSharedPtr<IMediaPlayer, ESPMode::ThreadSafe> NewPlayer = GetPlayerForUrl(Url, Options);
 
@@ -459,9 +465,15 @@ bool FMediaPlayerFacade::Open(const FString& Url, const IMediaOptions* Options,
 
 	if (!Player.IsValid())
 	{
+		// Make sure we don't get called from the "tickable" thread anymore - no need as we have no player
+		MediaModule->GetTicker().RemoveTickable(AsShared());
 		return false;
 	}
 
+	// Make sure we get ticked on the "tickable" thread
+	// (this will not re-add us, should we already be registered)
+	MediaModule->GetTicker().AddTickable(AsShared());
+
 	// update the Guid
 	Player->SetGuid(PlayerGuid);
 
@@ -897,6 +909,19 @@ void FMediaPlayerFacade::ProcessEvent(EMediaEvent Event)
 			FlushSinks();
 		}
 	}
+	else if (Event == EMediaEvent::MediaClosed)
+	{
+		// If player allows: close it down all the way right now
+		if (Player.IsValid() && Player->GetPlayerFeatureFlag(IMediaPlayer::FeatureFlag::AllowShutdownOnClose))
+		{
+			Player.Reset();
+		}
+
+		// Stop issuing audio thread ticks until we open the player again
+		IMediaModule* MediaModule = FModuleManager::LoadModulePtr<IMediaModule>("Media");
+		check(MediaModule);
+		MediaModule->GetTicker().RemoveTickable(AsShared());
+	}
 
 	MediaEvent.Broadcast(Event);
 }
diff --git a/Engine/Source/Runtime/MediaUtils/Public/MediaPlayerFacade.h b/Engine/Source/Runtime/MediaUtils/Public/MediaPlayerFacade.h
--- a/Engine/Source/Runtime/MediaUtils/Public/MediaPlayerFacade.h
+++ b/Engine/Source/Runtime/MediaUtils/Public/MediaPlayerFacade.h
@@ -50,6 +50,7 @@ class MEDIAUTILS_API FMediaPlayerFacade
 	: public IMediaClockSink
 	, public IMediaTickable
 	, protected IMediaEventSink
+	, public TSharedFromThis<FMediaPlayerFacade, ESPMode::ThreadSafe>
 {
 public:
 
