Fixing up open-open-play style use case for player facade & newly added tick/close optimizations
#rnx

diff --git a/Engine/Source/Runtime/MediaUtils/Private/MediaPlayerFacade.cpp b/Engine/Source/Runtime/MediaUtils/Private/MediaPlayerFacade.cpp
--- a/Engine/Source/Runtime/MediaUtils/Private/MediaPlayerFacade.cpp
+++ b/Engine/Source/Runtime/MediaUtils/Private/MediaPlayerFacade.cpp
@@ -911,16 +911,23 @@ void FMediaPlayerFacade::ProcessEvent(EMediaEvent Event)
 	}
 	else if (Event == EMediaEvent::MediaClosed)
 	{
-		// If player allows: close it down all the way right now
-		if (Player.IsValid() && Player->GetPlayerFeatureFlag(IMediaPlayer::FeatureFlag::AllowShutdownOnClose))
+		// Player still closed?
+		if (CurrentUrl.IsEmpty())
 		{
-			Player.Reset();
-		}
+			// Yes, this also means: if we still have a player, it's still the one this event originated from
 
-		// Stop issuing audio thread ticks until we open the player again
-		IMediaModule* MediaModule = FModuleManager::LoadModulePtr<IMediaModule>("Media");
-		check(MediaModule);
-		MediaModule->GetTicker().RemoveTickable(AsShared());
+			// If player allows: close it down all the way right now
+			if (Player.IsValid() && Player->GetPlayerFeatureFlag(IMediaPlayer::FeatureFlag::AllowShutdownOnClose))
+			{
+				FScopeLock Lock(&CriticalSection);
+				Player.Reset();
+			}
+
+			// Stop issuing audio thread ticks until we open the player again
+			IMediaModule* MediaModule = FModuleManager::LoadModulePtr<IMediaModule>("Media");
+			check(MediaModule);
+			MediaModule->GetTicker().RemoveTickable(AsShared());
+		}
 	}
 
 	MediaEvent.Broadcast(Event);
