Integration of reworked MWF output pipeline into Dev-Main

diff --git a/Engine/Plugins/Lumin/MagicLeapMedia/Source/MagicLeapMedia/Private/MagicLeapMediaOverlaySample.h b/Engine/Plugins/Lumin/MagicLeapMedia/Source/MagicLeapMedia/Private/MagicLeapMediaOverlaySample.h
--- a/Engine/Plugins/Lumin/MagicLeapMedia/Source/MagicLeapMedia/Private/MagicLeapMediaOverlaySample.h
+++ b/Engine/Plugins/Lumin/MagicLeapMedia/Source/MagicLeapMedia/Private/MagicLeapMediaOverlaySample.h
@@ -41,9 +41,9 @@ public:
 		return Text;
 	}
 
-	virtual FTimespan GetTime() const override
+	virtual FMediaTimeStamp GetTime() const override
 	{
-		return Time;
+		return FMediaTimeStamp(Time);
 	}
 
 	virtual EMediaOverlaySampleType GetType() const override
diff --git a/Engine/Plugins/Lumin/MagicLeapMedia/Source/MagicLeapMediaCodec/Private/MagicLeapMediaAudioSample.h b/Engine/Plugins/Lumin/MagicLeapMedia/Source/MagicLeapMediaCodec/Private/MagicLeapMediaAudioSample.h
--- a/Engine/Plugins/Lumin/MagicLeapMedia/Source/MagicLeapMediaCodec/Private/MagicLeapMediaAudioSample.h
+++ b/Engine/Plugins/Lumin/MagicLeapMedia/Source/MagicLeapMediaCodec/Private/MagicLeapMediaAudioSample.h
@@ -96,9 +96,9 @@ public:
 		return SampleRate;
 	}
 
-	virtual FTimespan GetTime() const override
+	virtual FMediaTimeStamp GetTime() const override
 	{
-		return Time;
+		return FMediaTimeStamp(Time);
 	}
 
 private:
diff --git a/Engine/Plugins/Media/AndroidCamera/Source/AndroidCamera/Private/Player/AndroidCameraTextureSample.h b/Engine/Plugins/Media/AndroidCamera/Source/AndroidCamera/Private/Player/AndroidCameraTextureSample.h
--- a/Engine/Plugins/Media/AndroidCamera/Source/AndroidCamera/Private/Player/AndroidCameraTextureSample.h
+++ b/Engine/Plugins/Media/AndroidCamera/Source/AndroidCamera/Private/Player/AndroidCameraTextureSample.h
@@ -218,9 +218,9 @@ public:
 
 #endif //WITH_ENGINE
 
-	virtual FTimespan GetTime() const override
+	virtual FMediaTimeStamp GetTime() const override
 	{
-		return Time;
+		return FMediaTimeStamp(Time);
 	}
 
 	virtual bool IsCacheable() const override
diff --git a/Engine/Plugins/Media/AndroidMedia/Source/AndroidMedia/Private/Player/AndroidMediaPlayer.cpp b/Engine/Plugins/Media/AndroidMedia/Source/AndroidMedia/Private/Player/AndroidMediaPlayer.cpp
--- a/Engine/Plugins/Media/AndroidMedia/Source/AndroidMedia/Private/Player/AndroidMediaPlayer.cpp
+++ b/Engine/Plugins/Media/AndroidMedia/Source/AndroidMedia/Private/Player/AndroidMediaPlayer.cpp
@@ -715,9 +715,18 @@ void FAndroidMediaPlayer::TickFetch(FTimespan DeltaTime, FTimespan /*Timecode*/)
 		{
 			TWeakPtr<FJavaAndroidMediaPlayer, ESPMode::ThreadSafe> JavaMediaPlayerPtr;
 			FGuid PlayerGuid;
+			TSharedRef<FAndroidMediaTextureSample, ESPMode::ThreadSafe> VideoSample;
+			TWeakPtr<FMediaSamples, ESPMode::ThreadSafe> Samples;
 		};
 
-		FWriteVideoSampleParams Params = { JavaMediaPlayer, PlayerGuid };
+		auto VideoSample = VideoSamplePool->AcquireShared();
+		const FJavaAndroidMediaPlayer::FVideoTrack VideoTrack = VideoTracks[SelectedVideoTrack];
+		if (!VideoSample->Initialize(VideoTrack.Dimensions, FTimespan::FromSeconds(1.0 / VideoTrack.FrameRate), FTimespan::FromMilliseconds(JavaMediaPlayer->GetCurrentPosition())))
+		{
+			return;
+		}
+			
+		FWriteVideoSampleParams Params = { JavaMediaPlayer, PlayerGuid, VideoSample, Samples };
 
 		ENQUEUE_RENDER_COMMAND(AndroidMediaPlayerWriteVideoSample)(
 			[Params](FRHICommandListImmediate& RHICmdList)
@@ -734,8 +743,9 @@ void FAndroidMediaPlayer::TickFetch(FTimespan DeltaTime, FTimespan /*Timecode*/)
 				}
 
 				auto PinnedJavaMediaPlayer = Params.JavaMediaPlayerPtr.Pin();
+				auto PinnedSamples = Params.Samples.Pin();
 
-				if (!PinnedJavaMediaPlayer.IsValid())
+				if (!PinnedJavaMediaPlayer.IsValid() || ! PinnedSamples.IsValid())
 				{
 					return;
 				}
@@ -752,7 +762,11 @@ void FAndroidMediaPlayer::TickFetch(FTimespan DeltaTime, FTimespan /*Timecode*/)
 
 					PinnedJavaMediaPlayer->SetVideoTextureValid(true);
 				}
-			});
+
+				// Add a dummy video sample
+				// (actual data flows "on the side" via external texture mechanism)
+				PinnedSamples->AddVideo(Params.VideoSample);
+		});
 	}
 #endif // ANDROIDMEDIAPLAYER_USE_EXTERNALTEXTURE
 	else
diff --git a/Engine/Plugins/Media/AndroidMedia/Source/AndroidMedia/Private/Player/AndroidMediaTextureSample.h b/Engine/Plugins/Media/AndroidMedia/Source/AndroidMedia/Private/Player/AndroidMediaTextureSample.h
--- a/Engine/Plugins/Media/AndroidMedia/Source/AndroidMedia/Private/Player/AndroidMediaTextureSample.h
+++ b/Engine/Plugins/Media/AndroidMedia/Source/AndroidMedia/Private/Player/AndroidMediaTextureSample.h
@@ -66,7 +66,7 @@ public:
 	 * @param InDuration The duration for which the sample is valid.
 	 * @return true on success, false otherwise.
 	 */
-	bool Initialize(const FIntPoint& InDim, FTimespan InDuration)
+	bool Initialize(const FIntPoint& InDim, FTimespan InDuration, FTimespan InTime = FTimespan::Zero())
 	{
 		if (InDim.GetMin() <= 0)
 		{
@@ -75,6 +75,7 @@ public:
 
 		Dim = InDim;
 		Duration = InDuration;
+		Time = InTime;
 
 		return true;
 	}
@@ -218,9 +219,9 @@ public:
 
 #endif //WITH_ENGINE
 
-	virtual FTimespan GetTime() const override
+	virtual FMediaTimeStamp GetTime() const override
 	{
-		return Time;
+		return FMediaTimeStamp(Time);
 	}
 
 	void UpdateTimeDuration(FTimespan InTime, FTimespan InDuration)
diff --git a/Engine/Plugins/Media/AvfMedia/Source/AvfMedia/Private/Player/AvfMediaAudioSample.h b/Engine/Plugins/Media/AvfMedia/Source/AvfMedia/Private/Player/AvfMediaAudioSample.h
--- a/Engine/Plugins/Media/AvfMedia/Source/AvfMedia/Private/Player/AvfMediaAudioSample.h
+++ b/Engine/Plugins/Media/AvfMedia/Source/AvfMedia/Private/Player/AvfMediaAudioSample.h
@@ -111,9 +111,9 @@ public:
 		return SampleRate;
 	}
 
-	virtual FTimespan GetTime() const override
+	virtual FMediaTimeStamp GetTime() const override
 	{
-		return Time;
+		return FMediaTimeStamp(Time);
 	}
 
 private:
diff --git a/Engine/Plugins/Media/AvfMedia/Source/AvfMedia/Private/Player/AvfMediaOverlaySample.h b/Engine/Plugins/Media/AvfMedia/Source/AvfMedia/Private/Player/AvfMediaOverlaySample.h
--- a/Engine/Plugins/Media/AvfMedia/Source/AvfMedia/Private/Player/AvfMediaOverlaySample.h
+++ b/Engine/Plugins/Media/AvfMedia/Source/AvfMedia/Private/Player/AvfMediaOverlaySample.h
@@ -59,9 +59,9 @@ public:
 		return Text;
 	}
 
-	virtual FTimespan GetTime() const override
+	virtual FMediaTimeStamp GetTime() const override
 	{
-		return Time;
+		return FMediaTimeStamp(Time);
 	}
 
 	virtual EMediaOverlaySampleType GetType() const override
diff --git a/Engine/Plugins/Media/AvfMedia/Source/AvfMedia/Private/Player/AvfMediaPlayer.cpp b/Engine/Plugins/Media/AvfMedia/Source/AvfMedia/Private/Player/AvfMediaPlayer.cpp
--- a/Engine/Plugins/Media/AvfMedia/Source/AvfMedia/Private/Player/AvfMediaPlayer.cpp
+++ b/Engine/Plugins/Media/AvfMedia/Source/AvfMedia/Private/Player/AvfMediaPlayer.cpp
@@ -261,7 +261,7 @@ public:
 		
 		if(FTimespan::MinValue() == AudioSyncSampleTime && bResult && OutSample.IsValid())
 		{
-			AudioSyncSampleTime = OutSample->GetTime() + OutSample->GetDuration();
+			AudioSyncSampleTime = OutSample->GetTime().Time + OutSample->GetDuration();
 		}
 
 		return bResult;
@@ -273,7 +273,7 @@ public:
 		
 		if(FTimespan::MinValue() == VideoSyncSampleTime && bResult && OutSample.IsValid())
 		{
-			VideoSyncSampleTime = OutSample->GetTime() + OutSample->GetDuration();
+			VideoSyncSampleTime = OutSample->GetTime().Time + OutSample->GetDuration();
 		}
 
 		return bResult;
diff --git a/Engine/Plugins/Media/AvfMedia/Source/AvfMedia/Private/Player/AvfMediaTextureSample.h b/Engine/Plugins/Media/AvfMedia/Source/AvfMedia/Private/Player/AvfMediaTextureSample.h
--- a/Engine/Plugins/Media/AvfMedia/Source/AvfMedia/Private/Player/AvfMediaTextureSample.h
+++ b/Engine/Plugins/Media/AvfMedia/Source/AvfMedia/Private/Player/AvfMediaTextureSample.h
@@ -148,9 +148,9 @@ public:
 		return Texture;
 	}
 
-	virtual FTimespan GetTime() const override
+	virtual FMediaTimeStamp GetTime() const override
 	{
-		return Time;
+		return FMediaTimeStamp(Time);
 	}
 
 	virtual bool IsCacheable() const override
diff --git a/Engine/Plugins/Media/AvfMedia/Source/AvfMedia/Private/Player/AvfMediaVideoSampler.h b/Engine/Plugins/Media/AvfMedia/Source/AvfMedia/Private/Player/AvfMediaVideoSampler.h
--- a/Engine/Plugins/Media/AvfMedia/Source/AvfMedia/Private/Player/AvfMediaVideoSampler.h
+++ b/Engine/Plugins/Media/AvfMedia/Source/AvfMedia/Private/Player/AvfMediaVideoSampler.h
@@ -16,7 +16,7 @@ class FMediaSamples;
 /**
  * Creates samples from video frames.
  */
-class FAvfMediaVideoSampler
+class DLLEXPORT FAvfMediaVideoSampler
 {
 public:
 
diff --git a/Engine/Plugins/Media/ImgMedia/Source/ImgMedia/Private/Player/ImgMediaPlayer.cpp b/Engine/Plugins/Media/ImgMedia/Source/ImgMedia/Private/Player/ImgMediaPlayer.cpp
--- a/Engine/Plugins/Media/ImgMedia/Source/ImgMedia/Private/Player/ImgMediaPlayer.cpp
+++ b/Engine/Plugins/Media/ImgMedia/Source/ImgMedia/Private/Player/ImgMediaPlayer.cpp
@@ -591,7 +591,7 @@ bool FImgMediaPlayer::FetchVideo(TRange<FTimespan> TimeRange, TSharedPtr<IMediaT
 		return false; // sample not loaded yet
 	}
 
-	const FTimespan SampleTime = Sample->GetTime();
+	const FTimespan SampleTime = Sample->GetTime().Time;
 
 	if (SampleTime == LastFetchTime)
 	{
@@ -611,6 +611,12 @@ void FImgMediaPlayer::FlushSamples()
 }
 
 
+bool FImgMediaPlayer::PeekVideoSampleTime(FMediaTimeStamp & TimeStamp)
+{
+	// player does not support v2 timing control at this point -> no need for this method, yet
+	return false;
+}
+
 /* IMediaTracks interface
  *****************************************************************************/
 
diff --git a/Engine/Plugins/Media/ImgMedia/Source/ImgMedia/Private/Player/ImgMediaPlayer.h b/Engine/Plugins/Media/ImgMedia/Source/ImgMedia/Private/Player/ImgMediaPlayer.h
--- a/Engine/Plugins/Media/ImgMedia/Source/ImgMedia/Private/Player/ImgMediaPlayer.h
+++ b/Engine/Plugins/Media/ImgMedia/Source/ImgMedia/Private/Player/ImgMediaPlayer.h
@@ -99,6 +99,7 @@ protected:
 
 	virtual bool FetchVideo(TRange<FTimespan> TimeRange, TSharedPtr<IMediaTextureSample, ESPMode::ThreadSafe>& OutSample) override;
 	virtual void FlushSamples() override;
+	virtual bool PeekVideoSampleTime(FMediaTimeStamp & TimeStamp) override;
 
 protected:
 
diff --git a/Engine/Plugins/Media/ImgMedia/Source/ImgMedia/Private/Player/ImgMediaTextureSample.h b/Engine/Plugins/Media/ImgMedia/Source/ImgMedia/Private/Player/ImgMediaTextureSample.h
--- a/Engine/Plugins/Media/ImgMedia/Source/ImgMedia/Private/Player/ImgMediaTextureSample.h
+++ b/Engine/Plugins/Media/ImgMedia/Source/ImgMedia/Private/Player/ImgMediaTextureSample.h
@@ -87,9 +87,9 @@ public:
 	}
 #endif //WITH_ENGINE
 
-	virtual FTimespan GetTime() const override
+	virtual FMediaTimeStamp GetTime() const override
 	{
-		return Time;
+		return FMediaTimeStamp(Time);
 	}
 
 	virtual bool IsCacheable() const override
diff --git a/Engine/Plugins/Media/MediaIOFramework/Source/MediaIOCore/Private/Player/MediaIOCoreSamples.cpp b/Engine/Plugins/Media/MediaIOFramework/Source/MediaIOCore/Private/Player/MediaIOCoreSamples.cpp
--- a/Engine/Plugins/Media/MediaIOFramework/Source/MediaIOCore/Private/Player/MediaIOCoreSamples.cpp
+++ b/Engine/Plugins/Media/MediaIOFramework/Source/MediaIOCore/Private/Player/MediaIOCoreSamples.cpp
@@ -103,3 +103,10 @@ void FMediaIOCoreSamples::FlushSamples()
 	SubtitleSamples.FlushSamples();
 	VideoSamples.FlushSamples();
 }
+
+
+bool FMediaIOCoreSamples::PeekVideoSampleTime(FMediaTimeStamp & TimeStamp)
+{
+	// player does not support v2 timing control at this point -> no need for this method, yet
+	return false;
+}
\ No newline at end of file
diff --git a/Engine/Plugins/Media/MediaIOFramework/Source/MediaIOCore/Public/MediaIOCoreAudioSampleBase.h b/Engine/Plugins/Media/MediaIOFramework/Source/MediaIOCore/Public/MediaIOCoreAudioSampleBase.h
--- a/Engine/Plugins/Media/MediaIOFramework/Source/MediaIOCore/Public/MediaIOCoreAudioSampleBase.h
+++ b/Engine/Plugins/Media/MediaIOFramework/Source/MediaIOCore/Public/MediaIOCoreAudioSampleBase.h
@@ -110,9 +110,9 @@ public:
 		return SampleRate;
 	}
 
-	virtual FTimespan GetTime() const override
+	virtual FMediaTimeStamp GetTime() const override
 	{
-		return Time;
+		return FMediaTimeStamp(Time);
 	}
 
 	virtual TOptional<FTimecode> GetTimecode() const override
diff --git a/Engine/Plugins/Media/MediaIOFramework/Source/MediaIOCore/Public/MediaIOCoreBinarySampleBase.h b/Engine/Plugins/Media/MediaIOFramework/Source/MediaIOCore/Public/MediaIOCoreBinarySampleBase.h
--- a/Engine/Plugins/Media/MediaIOFramework/Source/MediaIOCore/Public/MediaIOCoreBinarySampleBase.h
+++ b/Engine/Plugins/Media/MediaIOFramework/Source/MediaIOCore/Public/MediaIOCoreBinarySampleBase.h
@@ -89,9 +89,9 @@ public:
 		return Buffer.Num();
 	}
 
-	virtual FTimespan GetTime() const override
+	virtual FMediaTimeStamp GetTime() const override
 	{
-		return Time;
+		return FMediaTimeStamp(Time);
 	}
 
 	virtual TOptional<FTimecode> GetTimecode() const override
diff --git a/Engine/Plugins/Media/MediaIOFramework/Source/MediaIOCore/Public/MediaIOCoreSampleContainer.h b/Engine/Plugins/Media/MediaIOFramework/Source/MediaIOCore/Public/MediaIOCoreSampleContainer.h
--- a/Engine/Plugins/Media/MediaIOFramework/Source/MediaIOCore/Public/MediaIOCoreSampleContainer.h
+++ b/Engine/Plugins/Media/MediaIOFramework/Source/MediaIOCore/Public/MediaIOCoreSampleContainer.h
@@ -73,7 +73,7 @@ public:
 				if (Sample.IsValid())
 				{
 					FTimedDataChannelSampleTime& NewSampleTime = CachedSamplesData.Emplace_GetRef();
-					NewSampleTime.PlatformSecond = Sample->GetTime().GetTotalSeconds() - EvaluationSettings.PlayerTimeOffset;
+					NewSampleTime.PlatformSecond = Sample->GetTime().Time.GetTotalSeconds() - EvaluationSettings.PlayerTimeOffset;
 					if (Sample->GetTimecode().IsSet())
 					{
 						NewSampleTime.Timecode = FQualifiedFrameTime(Sample->GetTimecode().GetValue(), EvaluationSettings.FrameRate);
@@ -104,8 +104,8 @@ public:
 				}
 				else //Platform time
 				{
-					NewestSampleInSeconds = Samples[0]->GetTime().GetTotalSeconds() + Duration;
-					OldestSampleInSeconds = Samples[CachedSamplesData.Num() - 1]->GetTime().GetTotalSeconds();
+					NewestSampleInSeconds = Samples[0]->GetTime().Time.GetTotalSeconds() + Duration;
+					OldestSampleInSeconds = Samples[CachedSamplesData.Num() - 1]->GetTime().Time.GetTotalSeconds();
 				}
 
 				//Compute distance to evaluation taking into account duration of our samples for the newest one.
@@ -223,7 +223,7 @@ public:
 			TSharedPtr<SampleType, ESPMode::ThreadSafe> Sample = Samples[SampleCount - 1];
 			if (Sample.IsValid())
 			{
-				return Sample->GetTime();
+				return Sample->GetTime().Time;
 			}
 		}
 
@@ -354,7 +354,7 @@ public:
 			TSharedPtr<SampleType, ESPMode::ThreadSafe> Sample = Samples[SampleCount - 1];
 			if (Sample.IsValid())
 			{
-				const FTimespan SampleTime = Sample->GetTime();
+				const FTimespan SampleTime = Sample->GetTime().Time;
 
 				if (TimeRange.Overlaps(TRange<FTimespan>(SampleTime, SampleTime + Sample->GetDuration())))
 				{
diff --git a/Engine/Plugins/Media/MediaIOFramework/Source/MediaIOCore/Public/MediaIOCoreSamples.h b/Engine/Plugins/Media/MediaIOFramework/Source/MediaIOCore/Public/MediaIOCoreSamples.h
--- a/Engine/Plugins/Media/MediaIOFramework/Source/MediaIOCore/Public/MediaIOCoreSamples.h
+++ b/Engine/Plugins/Media/MediaIOFramework/Source/MediaIOCore/Public/MediaIOCoreSamples.h
@@ -297,6 +297,7 @@ public:
 	virtual bool FetchSubtitle(TRange<FTimespan> TimeRange, TSharedPtr<IMediaOverlaySample, ESPMode::ThreadSafe>& OutSample) override;
 	virtual bool FetchVideo(TRange<FTimespan> TimeRange, TSharedPtr<IMediaTextureSample, ESPMode::ThreadSafe>& OutSample) override;
 	virtual void FlushSamples() override;
+	virtual bool PeekVideoSampleTime(FMediaTimeStamp & TimeStamp) override;
 
 protected:
 
diff --git a/Engine/Plugins/Media/MediaIOFramework/Source/MediaIOCore/Public/MediaIOCoreTextureSampleBase.h b/Engine/Plugins/Media/MediaIOFramework/Source/MediaIOCore/Public/MediaIOCoreTextureSampleBase.h
--- a/Engine/Plugins/Media/MediaIOFramework/Source/MediaIOCore/Public/MediaIOCoreTextureSampleBase.h
+++ b/Engine/Plugins/Media/MediaIOFramework/Source/MediaIOCore/Public/MediaIOCoreTextureSampleBase.h
@@ -165,9 +165,9 @@ public:
 	}
 #endif //WITH_ENGINE
 
-	virtual FTimespan GetTime() const override
+	virtual FMediaTimeStamp GetTime() const override
 	{
-		return Time;
+		return FMediaTimeStamp(Time);
 	}
 
 	virtual TOptional<FTimecode> GetTimecode() const override
diff --git a/Engine/Plugins/Media/MediaPlayerEditor/Source/MediaPlayerEditor/Private/Customizations/MediaTextureCustomization.cpp b/Engine/Plugins/Media/MediaPlayerEditor/Source/MediaPlayerEditor/Private/Customizations/MediaTextureCustomization.cpp
--- a/Engine/Plugins/Media/MediaPlayerEditor/Source/MediaPlayerEditor/Private/Customizations/MediaTextureCustomization.cpp
+++ b/Engine/Plugins/Media/MediaPlayerEditor/Source/MediaPlayerEditor/Private/Customizations/MediaTextureCustomization.cpp
@@ -17,9 +17,7 @@ void FMediaTextureCustomization::CustomizeDetails(IDetailLayoutBuilder& DetailBu
 	// customize 'Texture' category
 	IDetailCategoryBuilder& MediaTextureCategory = DetailBuilder.EditCategory("MediaTexture");
 	{
-		MediaTextureCategory.AddProperty(GET_MEMBER_NAME_CHECKED(UTexture, CompressionSettings), UTexture::StaticClass());
 		MediaTextureCategory.AddProperty(GET_MEMBER_NAME_CHECKED(UTexture, Filter), UTexture::StaticClass());
-		MediaTextureCategory.AddProperty(GET_MEMBER_NAME_CHECKED(UTexture, SRGB), UTexture::StaticClass());
 	}
 
 	DetailBuilder.HideCategory("Compression");
diff --git a/Engine/Plugins/Media/MfMedia/Source/MfMedia/Private/Player/MfMediaAudioSample.h b/Engine/Plugins/Media/MfMedia/Source/MfMedia/Private/Player/MfMediaAudioSample.h
--- a/Engine/Plugins/Media/MfMedia/Source/MfMedia/Private/Player/MfMediaAudioSample.h
+++ b/Engine/Plugins/Media/MfMedia/Source/MfMedia/Private/Player/MfMediaAudioSample.h
@@ -134,9 +134,9 @@ public:
 		return SampleRate;
 	}
 
-	virtual FTimespan GetTime() const override
+	virtual FMediaTimeStamp GetTime() const override
 	{
-		return Time;
+		return FMediaTimeStamp(Time);
 	}
 
 private:
diff --git a/Engine/Plugins/Media/MfMedia/Source/MfMedia/Private/Player/MfMediaTextureSample.h b/Engine/Plugins/Media/MfMedia/Source/MfMedia/Private/Player/MfMediaTextureSample.h
--- a/Engine/Plugins/Media/MfMedia/Source/MfMedia/Private/Player/MfMediaTextureSample.h
+++ b/Engine/Plugins/Media/MfMedia/Source/MfMedia/Private/Player/MfMediaTextureSample.h
@@ -192,9 +192,9 @@ public:
 	}
 #endif //WITH_ENGINE
 
-	virtual FTimespan GetTime() const override
+	virtual FMediaTimeStamp GetTime() const override
 	{
-		return Time;
+		return FMediaTimeStamp(Time);
 	}
 
 	virtual bool IsCacheable() const override
diff --git a/Engine/Plugins/Media/MfMedia/Source/MfMedia/Private/Player/MfMediaTracks.cpp b/Engine/Plugins/Media/MfMedia/Source/MfMedia/Private/Player/MfMediaTracks.cpp
--- a/Engine/Plugins/Media/MfMedia/Source/MfMedia/Private/Player/MfMediaTracks.cpp
+++ b/Engine/Plugins/Media/MfMedia/Source/MfMedia/Private/Player/MfMediaTracks.cpp
@@ -302,7 +302,7 @@ void FMfMediaTracks::ProcessSample(IMFSample* Sample, HRESULT Status, DWORD Stre
 				if (AudioSample->Initialize(*Format.OutputType, *Sample, Format.Audio.NumChannels, Format.Audio.SampleRate))
 				{
 					Samples->AddAudio(AudioSample);
-					LastAudioSampleTime = AudioSample->GetTime();
+					LastAudioSampleTime = AudioSample->GetTime().Time;
 
 					#if MFMEDIATRACKS_TRACE_SAMPLES
 						UE_LOG(LogMfMedia, VeryVerbose, TEXT("Tracks %p: Audio sample processed: %s"), this, *LastAudioSampleTime.ToString());
@@ -392,7 +392,7 @@ void FMfMediaTracks::ProcessSample(IMFSample* Sample, HRESULT Status, DWORD Stre
 				if (VideoSample->Initialize(*Format.OutputType, *Sample, Format.Video.BufferDim, Format.Video.BufferStride, Format.Video.OutputDim, true))
 				{
 					Samples->AddVideo(VideoSample);
-					LastVideoSampleTime = VideoSample->GetTime();
+					LastVideoSampleTime = VideoSample->GetTime().Time;
 
 					#if MFMEDIATRACKS_TRACE_SAMPLES
 						UE_LOG(LogMfMedia, VeryVerbose, TEXT("Tracks %p: Video sample processed: %s"), this, *LastVideoSampleTime.ToString());
diff --git a/Engine/Plugins/Media/PixelStreaming/Source/PixelStreaming/Private/AudioSink.h b/Engine/Plugins/Media/PixelStreaming/Source/PixelStreaming/Private/AudioSink.h
--- a/Engine/Plugins/Media/PixelStreaming/Source/PixelStreaming/Private/AudioSink.h
+++ b/Engine/Plugins/Media/PixelStreaming/Source/PixelStreaming/Private/AudioSink.h
@@ -61,9 +61,9 @@ public:
 		return SampleRate;
 	}
 
-	FTimespan GetTime() const override
+	FMediaTimeStamp GetTime() const override
 	{
-		return Timestamp;
+		return FMediaTimeStamp(Timestamp);
 	}
 
 private:
diff --git a/Engine/Plugins/Media/PixelStreaming/Source/PixelStreaming/Private/Codecs/SwTextureSample.h b/Engine/Plugins/Media/PixelStreaming/Source/PixelStreaming/Private/Codecs/SwTextureSample.h
--- a/Engine/Plugins/Media/PixelStreaming/Source/PixelStreaming/Private/Codecs/SwTextureSample.h
+++ b/Engine/Plugins/Media/PixelStreaming/Source/PixelStreaming/Private/Codecs/SwTextureSample.h
@@ -125,9 +125,9 @@ public:
 	}
 #endif //WITH_ENGINE
 
-	FTimespan GetTime() const override
+	FMediaTimeStamp GetTime() const override
 	{
-		return Time;
+		return FMediaTimeStamp(Time);
 	}
 
 	bool IsCacheable() const override
diff --git a/Engine/Plugins/Media/PixelStreaming/Source/PixelStreaming/Private/Codecs/VideoDecoder.cpp b/Engine/Plugins/Media/PixelStreaming/Source/PixelStreaming/Private/Codecs/VideoDecoder.cpp
--- a/Engine/Plugins/Media/PixelStreaming/Source/PixelStreaming/Private/Codecs/VideoDecoder.cpp
+++ b/Engine/Plugins/Media/PixelStreaming/Source/PixelStreaming/Private/Codecs/VideoDecoder.cpp
@@ -581,12 +581,12 @@ bool FVideoDecoder::ProcessOutputHW(IMFSample* MFSample)
 	check(DecodeCallback);
 	webrtc::VideoFrame VideoFrame = webrtc::VideoFrame::Builder{}.
 		set_video_frame_buffer(new rtc::RefCountedObject<FVideoFrameBuffer>(TextureSample)).
-		set_timestamp_rtp(TextureSample->GetTime().GetTicks()).
+		set_timestamp_rtp(TextureSample->GetTime().Time.GetTicks()).
 		build();
 
 	DecodeCallback->Decoded(VideoFrame);
 
-	UE_LOG(LogVideoDecoder, VeryVerbose, TEXT("ProcessOutputHW: #%d (%d), ts %lld, capture ts %lld"), OutputFrameProcessedCount, InputFrameProcessedCount - OutputFrameProcessedCount, TextureSample->GetTime().GetTicks(), TextureSample->GetDuration().GetTicks());
+	UE_LOG(LogVideoDecoder, VeryVerbose, TEXT("ProcessOutputHW: #%d (%d), ts %lld, capture ts %lld"), OutputFrameProcessedCount, InputFrameProcessedCount - OutputFrameProcessedCount, TextureSample->GetTime().Time.GetTicks(), TextureSample->GetDuration().GetTicks());
 	++OutputFrameProcessedCount;
 
 	// decoder `AddRef`ed it for us. Not releasing it leads to decoder hanging
@@ -608,12 +608,12 @@ bool FVideoDecoder::ProcessOutputSW(const FSwTextureSampleRef& TextureSample)
 	check(DecodeCallback);
 	webrtc::VideoFrame VideoFrame = webrtc::VideoFrame::Builder{}.
 		set_video_frame_buffer(new rtc::RefCountedObject<FVideoFrameBuffer>(TextureSample)).
-		set_timestamp_rtp(TextureSample->GetTime().GetTicks()).
+		set_timestamp_rtp(TextureSample->GetTime().Time.GetTicks()).
 		build();
 
 	DecodeCallback->Decoded(VideoFrame);
 
-	UE_LOG(LogVideoDecoder, VeryVerbose, TEXT("ProcessOutputSW: #%d, ts %.3f, d %.3f"), OutputFrameProcessedCount, TextureSample->GetTime().GetTotalSeconds(), TextureSample->GetDuration().GetTotalSeconds());
+	UE_LOG(LogVideoDecoder, VeryVerbose, TEXT("ProcessOutputSW: #%d, ts %.3f, d %.3f"), OutputFrameProcessedCount, TextureSample->GetTime().Time.GetTotalSeconds(), TextureSample->GetDuration().GetTotalSeconds());
 	++OutputFrameProcessedCount;
 	return true;
 }
diff --git a/Engine/Plugins/Media/PixelStreaming/Source/PixelStreaming/Private/Player.cpp b/Engine/Plugins/Media/PixelStreaming/Source/PixelStreaming/Private/Player.cpp
--- a/Engine/Plugins/Media/PixelStreaming/Source/PixelStreaming/Private/Player.cpp
+++ b/Engine/Plugins/Media/PixelStreaming/Source/PixelStreaming/Private/Player.cpp
@@ -211,6 +211,6 @@ void FPlayer::OnVideoFrame(const TSharedRef<IMediaTextureSample, ESPMode::Thread
 		FHUDStats::Get().EndToEndLatencyMs.Update(LatencyMs);
 	}
 
-	UE_LOG(PixelPlayer, Verbose, TEXT("(%d) Sending frame for rendering: ts %lld, capture ts %lld, latency %.0f"), RtcTimeMs(), Sample->GetTime().GetTicks(), Sample->GetDuration().GetTicks(), LatencyMs);
+	UE_LOG(PixelPlayer, Verbose, TEXT("(%d) Sending frame for rendering: ts %lld, capture ts %lld, latency %.0f"), RtcTimeMs(), Sample->GetTime().Time.GetTicks(), Sample->GetDuration().GetTicks(), LatencyMs);
 	MediaSamples->AddVideo(Sample);
 }
diff --git a/Engine/Plugins/Media/WebMMedia/Source/WebMMedia/Private/Player/WebMMediaTextureSample.cpp b/Engine/Plugins/Media/WebMMedia/Source/WebMMedia/Private/Player/WebMMediaTextureSample.cpp
--- a/Engine/Plugins/Media/WebMMedia/Source/WebMMedia/Private/Player/WebMMediaTextureSample.cpp
+++ b/Engine/Plugins/Media/WebMMedia/Source/WebMMedia/Private/Player/WebMMediaTextureSample.cpp
@@ -79,9 +79,9 @@ FRHITexture* FWebMMediaTextureSample::GetTexture() const
 	return Texture.GetReference();
 }
 
-FTimespan FWebMMediaTextureSample::GetTime() const
+FMediaTimeStamp FWebMMediaTextureSample::GetTime() const
 {
-	return Time;
+	return FMediaTimeStamp(Time);
 }
 
 bool FWebMMediaTextureSample::IsCacheable() const
diff --git a/Engine/Plugins/Media/WebMMedia/Source/WebMMedia/Public/WebMMediaAudioSample.h b/Engine/Plugins/Media/WebMMedia/Source/WebMMedia/Public/WebMMediaAudioSample.h
--- a/Engine/Plugins/Media/WebMMedia/Source/WebMMedia/Public/WebMMediaAudioSample.h
+++ b/Engine/Plugins/Media/WebMMedia/Source/WebMMedia/Public/WebMMediaAudioSample.h
@@ -94,9 +94,9 @@ public:
 		return SampleRate;
 	}
 
-	virtual FTimespan GetTime() const override
+	virtual FMediaTimeStamp GetTime() const override
 	{
-		return Time;
+		return FMediaTimeStamp(Time);
 	}
 
 	const TArray<uint8>& GetDataBuffer() const
diff --git a/Engine/Plugins/Media/WebMMedia/Source/WebMMedia/Public/WebMMediaTextureSample.h b/Engine/Plugins/Media/WebMMedia/Source/WebMMedia/Public/WebMMediaTextureSample.h
--- a/Engine/Plugins/Media/WebMMedia/Source/WebMMedia/Public/WebMMediaTextureSample.h
+++ b/Engine/Plugins/Media/WebMMedia/Source/WebMMedia/Public/WebMMediaTextureSample.h
@@ -30,7 +30,7 @@ public:
 	virtual FIntPoint GetOutputDim() const override;
 	virtual uint32 GetStride() const override;
 	virtual FRHITexture* GetTexture() const override;
-	virtual FTimespan GetTime() const override;
+	virtual FMediaTimeStamp GetTime() const override;
 	virtual bool IsCacheable() const override;
 	virtual bool IsOutputSrgb() const override;
 
diff --git a/Engine/Plugins/Media/WmfMedia/Source/WmfMedia/Private/Player/WmfMediaAudioSample.h b/Engine/Plugins/Media/WmfMedia/Source/WmfMedia/Private/Player/WmfMediaAudioSample.h
--- a/Engine/Plugins/Media/WmfMedia/Source/WmfMedia/Private/Player/WmfMediaAudioSample.h
+++ b/Engine/Plugins/Media/WmfMedia/Source/WmfMedia/Private/Player/WmfMediaAudioSample.h
@@ -99,9 +99,9 @@ public:
 		return SampleRate;
 	}
 
-	virtual FTimespan GetTime() const override
+	virtual FMediaTimeStamp GetTime() const override
 	{
-		return Time;
+		return FMediaTimeStamp(Time);
 	}
 
 private:
diff --git a/Engine/Plugins/Media/WmfMedia/Source/WmfMedia/Private/Player/WmfMediaBinarySample.h b/Engine/Plugins/Media/WmfMedia/Source/WmfMedia/Private/Player/WmfMediaBinarySample.h
--- a/Engine/Plugins/Media/WmfMedia/Source/WmfMedia/Private/Player/WmfMediaBinarySample.h
+++ b/Engine/Plugins/Media/WmfMedia/Source/WmfMedia/Private/Player/WmfMediaBinarySample.h
@@ -75,9 +75,9 @@ public:
 		return Buffer.Num();
 	}
 
-	virtual FTimespan GetTime() const override
+	virtual FMediaTimeStamp GetTime() const override
 	{
-		return Time;
+		return FMediaTimeStamp(Time);
 	}
 
 private:
diff --git a/Engine/Plugins/Media/WmfMedia/Source/WmfMedia/Private/Player/WmfMediaOverlaySample.h b/Engine/Plugins/Media/WmfMedia/Source/WmfMedia/Private/Player/WmfMediaOverlaySample.h
--- a/Engine/Plugins/Media/WmfMedia/Source/WmfMedia/Private/Player/WmfMediaOverlaySample.h
+++ b/Engine/Plugins/Media/WmfMedia/Source/WmfMedia/Private/Player/WmfMediaOverlaySample.h
@@ -97,9 +97,9 @@ public:
 		return Text;
 	}
 
-	virtual FTimespan GetTime() const override
+	virtual FMediaTimeStamp GetTime() const override
 	{
-		return Time;
+		return FMediaTimeStamp(Time);
 	}
 
 	virtual EMediaOverlaySampleType GetType() const override
diff --git a/Engine/Plugins/Media/WmfMedia/Source/WmfMedia/Private/Player/WmfMediaTextureSample.h b/Engine/Plugins/Media/WmfMedia/Source/WmfMedia/Private/Player/WmfMediaTextureSample.h
--- a/Engine/Plugins/Media/WmfMedia/Source/WmfMedia/Private/Player/WmfMediaTextureSample.h
+++ b/Engine/Plugins/Media/WmfMedia/Source/WmfMedia/Private/Player/WmfMediaTextureSample.h
@@ -121,9 +121,9 @@ public:
 	}
 #endif //WITH_ENGINE
 
-	virtual FTimespan GetTime() const override
+	virtual FMediaTimeStamp GetTime() const override
 	{
-		return Time;
+		return FMediaTimeStamp(Time);
 	}
 
 	virtual bool IsCacheable() const override
diff --git a/Engine/Plugins/Media/WmfMedia/Source/WmfMedia/Private/Player/WmfMediaTracks.cpp b/Engine/Plugins/Media/WmfMedia/Source/WmfMedia/Private/Player/WmfMediaTracks.cpp
--- a/Engine/Plugins/Media/WmfMedia/Source/WmfMedia/Private/Player/WmfMediaTracks.cpp
+++ b/Engine/Plugins/Media/WmfMedia/Source/WmfMedia/Private/Player/WmfMediaTracks.cpp
@@ -371,7 +371,7 @@ bool FWmfMediaTracks::FetchAudio(TRange<FTimespan> TimeRange, TSharedPtr<IMediaA
 		return false;
 	}
 
-	const FTimespan SampleTime = Sample->GetTime();
+	const FTimespan SampleTime = Sample->GetTime().Time;
 
 	if (!TimeRange.Overlaps(TRange<FTimespan>(SampleTime, SampleTime + Sample->GetDuration())))
 	{
@@ -398,7 +398,7 @@ bool FWmfMediaTracks::FetchCaption(TRange<FTimespan> TimeRange, TSharedPtr<IMedi
 		return false;
 	}
 
-	const FTimespan SampleTime = Sample->GetTime();
+	const FTimespan SampleTime = Sample->GetTime().Time;
 
 	if (!TimeRange.Overlaps(TRange<FTimespan>(SampleTime, SampleTime + Sample->GetDuration())))
 	{
@@ -425,7 +425,7 @@ bool FWmfMediaTracks::FetchMetadata(TRange<FTimespan> TimeRange, TSharedPtr<IMed
 		return false;
 	}
 
-	const FTimespan SampleTime = Sample->GetTime();
+	const FTimespan SampleTime = Sample->GetTime().Time;
 
 	if (!TimeRange.Overlaps(TRange<FTimespan>(SampleTime, SampleTime + Sample->GetDuration())))
 	{
@@ -452,7 +452,7 @@ bool FWmfMediaTracks::FetchVideo(TRange<FTimespan> TimeRange, TSharedPtr<IMediaT
 		return false;
 	}
 
-	const FTimespan SampleTime = Sample->GetTime();
+	const FTimespan SampleTime = Sample->GetTime().Time;
 
 	if (!TimeRange.Overlaps(TRange<FTimespan>(SampleTime, SampleTime + Sample->GetDuration())))
 	{
@@ -479,6 +479,17 @@ void FWmfMediaTracks::FlushSamples()
 }
 
 
+bool FWmfMediaTracks::PeekVideoSampleTime(FMediaTimeStamp & TimeStamp)
+{
+	TSharedPtr<IMediaTextureSample, ESPMode::ThreadSafe> Sample;
+	if (!VideoSampleQueue.Peek(Sample))
+	{
+		return false;
+	}
+	TimeStamp = FMediaTimeStamp(Sample->GetTime());
+	return true;
+}
+
 /* IMediaTracks interface
  *****************************************************************************/
 
diff --git a/Engine/Plugins/Media/WmfMedia/Source/WmfMedia/Private/Player/WmfMediaTracks.h b/Engine/Plugins/Media/WmfMedia/Source/WmfMedia/Private/Player/WmfMediaTracks.h
--- a/Engine/Plugins/Media/WmfMedia/Source/WmfMedia/Private/Player/WmfMediaTracks.h
+++ b/Engine/Plugins/Media/WmfMedia/Source/WmfMedia/Private/Player/WmfMediaTracks.h
@@ -27,6 +27,7 @@ class IMediaBinarySample;
 class IMediaOverlaySample;
 class IMediaTextureSample;
 class FWmfMediaTextureSample;
+class FMediaTimeStamp;
 
 enum class EMediaTextureSampleFormat;
 enum class EMediaTrackType;
@@ -187,6 +188,7 @@ public:
 	virtual bool FetchMetadata(TRange<FTimespan> TimeRange, TSharedPtr<IMediaBinarySample, ESPMode::ThreadSafe>& OutSample) override;
 	virtual bool FetchVideo(TRange<FTimespan> TimeRange, TSharedPtr<IMediaTextureSample, ESPMode::ThreadSafe>& OutSample) override;
 	virtual void FlushSamples() override;
+	virtual bool PeekVideoSampleTime(FMediaTimeStamp & TimeStamp) override;
 
 public:
 
diff --git a/Engine/Plugins/Media/WmfMedia/Source/WmfMedia/Public/WmfMediaHardwareVideoDecodingTextureSample.h b/Engine/Plugins/Media/WmfMedia/Source/WmfMedia/Public/WmfMediaHardwareVideoDecodingTextureSample.h
--- a/Engine/Plugins/Media/WmfMedia/Source/WmfMedia/Public/WmfMediaHardwareVideoDecodingTextureSample.h
+++ b/Engine/Plugins/Media/WmfMedia/Source/WmfMedia/Public/WmfMediaHardwareVideoDecodingTextureSample.h
@@ -63,9 +63,10 @@ public:
 	/**
 	 * Texture sample convert using hardware video decoding.
 	 */
-	virtual void Convert(FTexture2DRHIRef InDstTexture) override
+	virtual bool Convert(FTexture2DRHIRef & InDstTexture, const FConversionHints & Hints) override
 	{
 		FWmfMediaHardwareVideoDecodingParameters::ConvertTextureFormat_RenderThread(this, InDstTexture);
+		return true;
 	}
 
 	/**
diff --git a/Engine/Plugins/Runtime/WebMMoviePlayer/Source/WebMMoviePlayer/Private/WebMMovieStreamer.cpp b/Engine/Plugins/Runtime/WebMMoviePlayer/Source/WebMMoviePlayer/Private/WebMMovieStreamer.cpp
--- a/Engine/Plugins/Runtime/WebMMoviePlayer/Source/WebMMoviePlayer/Private/WebMMovieStreamer.cpp
+++ b/Engine/Plugins/Runtime/WebMMoviePlayer/Source/WebMMoviePlayer/Private/WebMMovieStreamer.cpp
@@ -282,7 +282,7 @@ bool FWebMMovieStreamer::SendAudio(float InDeltaTime)
 	while (bFoundSample && AudioSample)
 	{
 		FWebMMediaAudioSample* WebMSample = StaticCast<FWebMMediaAudioSample*>(AudioSample.Get());
-		AudioBackend->SendAudio(WebMSample->GetTime(), WebMSample->GetDataBuffer().GetData(), WebMSample->GetDataBuffer().Num());
+		AudioBackend->SendAudio(WebMSample->GetTime().Time, WebMSample->GetDataBuffer().GetData(), WebMSample->GetDataBuffer().Num());
 
 		bFoundSample = Samples->FetchAudio(TimeRange, AudioSample);
 	}
diff --git a/Engine/Shaders/Private/MediaShaders.usf b/Engine/Shaders/Private/MediaShaders.usf
--- a/Engine/Shaders/Private/MediaShaders.usf
+++ b/Engine/Shaders/Private/MediaShaders.usf
@@ -53,7 +53,7 @@ void BMPConvertPS(
 }
 
 
-// shader to convert a NV12 texture to RGBA
+// shader to convert a NV12 texture to RGBA (real NV12 texture - DX11)
 void NV12ConvertPS(
 	float4 InPosition : SV_POSITION,
 	float2 InUV : TEXCOORD0,
@@ -62,19 +62,37 @@ void NV12ConvertPS(
 	float2 ScaledUV = InUV * NV12ConvertUB.UVScale;
 
 	float3 YUV;
-	YUV.x = NV12ConvertUB.Texture.Sample(NV12ConvertUB.SamplerB, ScaledUV).x;
+	YUV.x = NV12ConvertUB.SRV_Y.Sample(NV12ConvertUB.SamplerB, ScaledUV).x;
+	YUV.yz = NV12ConvertUB.SRV_UV.Sample(NV12ConvertUB.SamplerP, ScaledUV).xy;
+
+	OutColor.xyz = YuvToRgb(YUV, NV12ConvertUB.ColorTransform, NV12ConvertUB.SrgbToLinear);
+	OutColor.w = 1.0f;
+}
+
+
+// shader to convert a NV12 texture to RGBA (NV12 data, packed into a G8 texture; manually addressed)
+void NV12ConvertAsBytesPS(
+	float4 InPosition : SV_POSITION,
+	float2 InUV : TEXCOORD0,
+	out float4 OutColor : SV_Target0)
+{
+	float2 ScaledUV = InUV * NV12ConvertAsBytesUB.UVScale;
+
+	float3 YUV;
+
+	YUV.x = NV12ConvertAsBytesUB.Texture.Sample(NV12ConvertAsBytesUB.SamplerB, ScaledUV).x;
 
 	ScaledUV.y *= 0.5f;
 	ScaledUV.y += 2.0f / 3.0f;
 
-	uint UVOffs = uint(InUV.x * NV12ConvertUB.OutputWidth / 2.0f) * 2;
-	float UPos = (float(UVOffs * NV12ConvertUB.UVScale.x) + 0.5f) / NV12ConvertUB.OutputWidth;
-	float VPos = (float(UVOffs * NV12ConvertUB.UVScale.x) + 1.5f) / NV12ConvertUB.OutputWidth;
+	uint UVOffs = uint(InUV.x * NV12ConvertAsBytesUB.OutputWidth / 2.0f) * 2;
+	float UPos = (float(UVOffs * NV12ConvertAsBytesUB.UVScale.x) + 0.5f) / NV12ConvertAsBytesUB.OutputWidth;
+	float VPos = (float(UVOffs * NV12ConvertAsBytesUB.UVScale.x) + 1.5f) / NV12ConvertAsBytesUB.OutputWidth;
 
-	YUV.y = NV12ConvertUB.Texture.Sample(NV12ConvertUB.SamplerP, float2(UPos, ScaledUV.y)).x;
-	YUV.z = NV12ConvertUB.Texture.Sample(NV12ConvertUB.SamplerP, float2(VPos, ScaledUV.y)).x;
+	YUV.y = NV12ConvertAsBytesUB.Texture.Sample(NV12ConvertAsBytesUB.SamplerP, float2(UPos, ScaledUV.y)).x;
+	YUV.z = NV12ConvertAsBytesUB.Texture.Sample(NV12ConvertAsBytesUB.SamplerP, float2(VPos, ScaledUV.y)).x;
 
-	OutColor.xyz = YuvToRgb(YUV, NV12ConvertUB.ColorTransform, NV12ConvertUB.SrgbToLinear);
+	OutColor.xyz = YuvToRgb(YUV, NV12ConvertAsBytesUB.ColorTransform, NV12ConvertAsBytesUB.SrgbToLinear);
 	OutColor.w = 1.0f;
 }
 
diff --git a/Engine/Source/Runtime/Media/Private/MediaModule.cpp b/Engine/Source/Runtime/Media/Private/MediaModule.cpp
--- a/Engine/Source/Runtime/Media/Private/MediaModule.cpp
+++ b/Engine/Source/Runtime/Media/Private/MediaModule.cpp
@@ -106,6 +106,8 @@ public:
 
 	virtual void TickPreEngine() override
 	{
+		FrameStartTime = FPlatformTime::Seconds();
+
 		if (TimeSource.IsValid())
 		{
 			Clock.UpdateTimecode(TimeSource->GetTimecode(), TimecodeLocked);
@@ -132,6 +134,11 @@ public:
 		PlayerFactories.Remove(&Factory);
 	}
 
+	virtual double GetFrameStartTime() const override
+	{
+		return FrameStartTime;
+	}
+
 public:
 
 	//~ IModuleInterface interface
@@ -170,6 +177,9 @@ private:
 	/** The media clock. */
 	FMediaClock Clock;
 
+	/** Realtime at which frame started. */
+	double FrameStartTime;
+
 	/** Time code of the current frame. */
 	FTimespan CurrentTimecode;
 
diff --git a/Engine/Source/Runtime/Media/Public/IMediaAudioSample.h b/Engine/Source/Runtime/Media/Public/IMediaAudioSample.h
--- a/Engine/Source/Runtime/Media/Public/IMediaAudioSample.h
+++ b/Engine/Source/Runtime/Media/Public/IMediaAudioSample.h
@@ -6,6 +6,7 @@
 #include "Misc/Optional.h"
 #include "Misc/Timecode.h"
 #include "Misc/Timespan.h"
+#include "IMediaTimeSource.h"
 
 
 /**
@@ -103,7 +104,7 @@ public:
 	 * @return Sample time.
 	 * @see GetBuffer, GetChannels, GetDuration, GetFormat, GetFrames, GetSampleRate
 	 */
-	virtual FTimespan GetTime() const = 0;
+	virtual FMediaTimeStamp GetTime() const = 0;
 
 	/**
 	 * Get the sample timecode if available.
diff --git a/Engine/Source/Runtime/Media/Public/IMediaBinarySample.h b/Engine/Source/Runtime/Media/Public/IMediaBinarySample.h
--- a/Engine/Source/Runtime/Media/Public/IMediaBinarySample.h
+++ b/Engine/Source/Runtime/Media/Public/IMediaBinarySample.h
@@ -6,6 +6,7 @@
 #include "Misc/Optional.h"
 #include "Misc/Timecode.h"
 #include "Misc/Timespan.h"
+#include "IMediaTimeSource.h"
 
 
 /**
@@ -49,7 +50,7 @@ public:
 	 * @return Sample time.
 	 * @see GetData, GetDuration, GetSize, GetTime
 	 */
-	virtual FTimespan GetTime() const = 0;
+	virtual FMediaTimeStamp GetTime() const = 0;
 
 	/**
 	 * Get the sample timecode if available.
diff --git a/Engine/Source/Runtime/Media/Public/IMediaEventSink.h b/Engine/Source/Runtime/Media/Public/IMediaEventSink.h
--- a/Engine/Source/Runtime/Media/Public/IMediaEventSink.h
+++ b/Engine/Source/Runtime/Media/Public/IMediaEventSink.h
@@ -36,7 +36,15 @@ enum class EMediaEvent
 	SeekCompleted,
 
 	/** Media tracks have changed. */
-	TracksChanged
+	TracksChanged,
+
+	// - - - - - - - - - - - - - - - - - - - - - - - -
+
+	/** Start of internal events not sent out to user code. */
+	Internal_Start,
+
+	/** Player requests purge of outdated video samples from queues. */
+	Internal_PurgeVideoSamplesHint = Internal_Start,
 };
 
 
diff --git a/Engine/Source/Runtime/Media/Public/IMediaModule.h b/Engine/Source/Runtime/Media/Public/IMediaModule.h
--- a/Engine/Source/Runtime/Media/Public/IMediaModule.h
+++ b/Engine/Source/Runtime/Media/Public/IMediaModule.h
@@ -118,6 +118,11 @@ public:
 	 */
 	virtual IMediaTicker& GetTicker() = 0;
 
+	/**
+	 * Get frame's processing approximate real time start time in seconds
+	 */
+	virtual double GetFrameStartTime() const = 0;
+
 	/**
 	 * Get a Delegate that is trigger once all MediaClockSink are TickInput
 	 *
diff --git a/Engine/Source/Runtime/Media/Public/IMediaOverlaySample.h b/Engine/Source/Runtime/Media/Public/IMediaOverlaySample.h
--- a/Engine/Source/Runtime/Media/Public/IMediaOverlaySample.h
+++ b/Engine/Source/Runtime/Media/Public/IMediaOverlaySample.h
@@ -7,7 +7,7 @@
 #include "Misc/Optional.h"
 #include "Misc/Timecode.h"
 #include "Misc/Timespan.h"
-
+#include "IMediaTimeSource.h"
 
 /**
  * Type of text overlay samples.
@@ -62,7 +62,7 @@ public:
 	 *
 	 * @return Sample time.
 	 */
-	virtual FTimespan GetTime() const = 0;
+	virtual FMediaTimeStamp GetTime() const = 0;
 
 	/**
 	 * Get the sample timecode if available.
diff --git a/Engine/Source/Runtime/Media/Public/IMediaPlayer.h b/Engine/Source/Runtime/Media/Public/IMediaPlayer.h
--- a/Engine/Source/Runtime/Media/Public/IMediaPlayer.h
+++ b/Engine/Source/Runtime/Media/Public/IMediaPlayer.h
@@ -305,11 +305,12 @@ public:
 		// Override in child class if needed.
 	}
 
-	enum class FeatureFlag {
+	enum class EFeatureFlag {
 		AllowShutdownOnClose = 0,	//!< Allow player to be shutdown right after 'close' event is received from it
+		UsePlaybackTimingV2,		//!< Use v2 playback timing and AV sync
 	};
 	
-	virtual bool GetPlayerFeatureFlag(FeatureFlag /*flag*/) const
+	virtual bool GetPlayerFeatureFlag(EFeatureFlag /*flag*/) const
 	{
 		// Override in child class if needed.
 		return false;
diff --git a/Engine/Source/Runtime/Media/Public/IMediaSamples.h b/Engine/Source/Runtime/Media/Public/IMediaSamples.h
--- a/Engine/Source/Runtime/Media/Public/IMediaSamples.h
+++ b/Engine/Source/Runtime/Media/Public/IMediaSamples.h
@@ -6,6 +6,7 @@
 #include "Math/Range.h"
 #include "Misc/Timespan.h"
 #include "Templates/SharedPointer.h"
+#include "IMediaTimeSource.h"
 
 class IMediaAudioSample;
 class IMediaBinarySample;
@@ -49,6 +50,10 @@ public:
 	{
 		return false; // override in child classes, if supported
 	}
+	virtual bool FetchCaption(TRange<FMediaTimeStamp> TimeRange, TSharedPtr<IMediaOverlaySample, ESPMode::ThreadSafe>& OutSample)
+	{
+		return false; // override in child classes, if supported
+	}
 
 	/**
 	 * Fetch the next metadata sample.
@@ -75,6 +80,10 @@ public:
 	{
 		return false; // override in child classes, if supported
 	}
+	virtual bool FetchSubtitle(TRange<FMediaTimeStamp> TimeRange, TSharedPtr<IMediaOverlaySample, ESPMode::ThreadSafe>& OutSample)
+	{
+		return false; // override in child classes, if supported
+	}
 
 	/**
 	 * Fetch the next video sample.
@@ -88,6 +97,10 @@ public:
 	{
 		return false; // override in child classes, if supported
 	}
+	virtual bool FetchVideo(TRange<FMediaTimeStamp> TimeRange, TSharedPtr<IMediaTextureSample, ESPMode::ThreadSafe>& OutSample)
+	{
+		return false; // override in child classes, if supported
+	}
 
 	/** Discard any outstanding media samples. */
 	virtual void FlushSamples()
@@ -95,6 +108,29 @@ public:
 		// override in child classes, if supported
 	}
 
+	enum class EFetchBestSampleResult
+	{
+		Ok = 0,
+		NoSample,
+		NotSupported,
+	};
+	virtual EFetchBestSampleResult FetchBestVideoSampleForTimeRange(const TRange<FMediaTimeStamp> & TimeRange, TSharedPtr<IMediaTextureSample, ESPMode::ThreadSafe>& OutSample, bool bReverse)
+	{
+		return EFetchBestSampleResult::NotSupported;
+	}
+
+	virtual bool PeekVideoSampleTime(FMediaTimeStamp & TimeStamp) = 0;
+
+	virtual uint32 PurgeOutdatedVideoSamples(const FMediaTimeStamp & ReferenceTime, bool bReversed) { return 0; };
+
+	virtual bool CanReceiveVideoSamples(uint32 Num) const { return true; }
+
+	virtual int32 NumAudio() const { return -1; }
+	virtual int32 NumCaption() const { return -1; }
+	virtual int32 NumMetadataSamples() const { return -1; }
+	virtual int32 NumSubtitleSamples() const { return -1; }
+	virtual int32 NumVideoSamples() const { return -1; }
+
 public:
 
 	/** Virtual destructor. */
diff --git a/Engine/Source/Runtime/Media/Public/IMediaTextureSample.h b/Engine/Source/Runtime/Media/Public/IMediaTextureSample.h
--- a/Engine/Source/Runtime/Media/Public/IMediaTextureSample.h
+++ b/Engine/Source/Runtime/Media/Public/IMediaTextureSample.h
@@ -11,6 +11,7 @@
 #include "Misc/Timecode.h"
 #include "Misc/Timespan.h"
 #include "Templates/SharedPointer.h"
+#include "IMediaTimeSource.h"
 
 #if WITH_ENGINE
 	class FRHITexture;
@@ -183,7 +184,7 @@ public:
 	 * @return Sample time.
 	 * @see GetBuffer, GetDim, GetDuration, GetFormat, GetOutputDim, GetStride, GetTexture
 	 */
-	virtual FTimespan GetTime() const = 0;
+	virtual FMediaTimeStamp GetTime() const = 0;
 
 	/**
 	 * Get the sample timecode if available.
diff --git a/Engine/Source/Runtime/Media/Public/IMediaTextureSampleConverter.h b/Engine/Source/Runtime/Media/Public/IMediaTextureSampleConverter.h
--- a/Engine/Source/Runtime/Media/Public/IMediaTextureSampleConverter.h
+++ b/Engine/Source/Runtime/Media/Public/IMediaTextureSampleConverter.h
@@ -11,5 +11,23 @@
 class IMediaTextureSampleConverter
 {
 public:
-	virtual void Convert(FTexture2DRHIRef InDstTexture) = 0;
+	enum EConverterInfoFlags
+	{
+		ConverterInfoFlags_Default = 0,
+		ConverterInfoFlags_WillCreateOutputTexture = 1 << 0,
+		ConverterInfoFlags_PreprocessOnly = 1 << 1,
+	};
+
+	struct FConversionHints
+	{
+		bool bOutputSRGB;
+		uint8 NumMips;
+	};
+
+	virtual uint32 GetConverterInfoFlags() const
+	{
+		return ConverterInfoFlags_Default;
+	}
+
+	virtual bool Convert(FTexture2DRHIRef & InDstTexture, const FConversionHints & Hints) = 0;
 };
diff --git a/Engine/Source/Runtime/Media/Public/IMediaTimeSource.h b/Engine/Source/Runtime/Media/Public/IMediaTimeSource.h
--- a/Engine/Source/Runtime/Media/Public/IMediaTimeSource.h
+++ b/Engine/Source/Runtime/Media/Public/IMediaTimeSource.h
@@ -4,6 +4,64 @@
 
 #include "Misc/Timespan.h"
 
+/*
+	Timestamp value for media playback
+
+	- Time:
+	  Time value
+
+	- SequenceIndex
+	  Sequence that is current for this time value
+
+	Within a single sequence time values will increase or decrease monotonically. A new sequence index is generated on each event that causes the time to no
+	longer be monotonic. (e.g. seek or loop)
+	A sequence index does not offer any information about the ordering of the time stamps on the timeline. Time values are comparable between all timestamps from a single playback,
+	though, one needs to be careful to consider non-monotonic behavior if the sequence indices are not identical.
+
+	Sequence indices can very much offer ordering information as far as playback progression is concerned. Higher indices are also later in playback. (even if time values may be smaller: e.g. looping)
+
+	All comparison operators of this class will operate to indicate ordering relative to playback, not position on the timeline!
+
+*/
+class FMediaTimeStamp
+{
+public:
+	FMediaTimeStamp() : Time(FTimespan::MinValue()), SequenceIndex(0) {}
+	FMediaTimeStamp(const FTimespan & InTime) : Time(InTime), SequenceIndex(0) {}
+	FMediaTimeStamp(const FTimespan & InTime, int64 InSequenceIndex) : Time(InTime), SequenceIndex(InSequenceIndex) {}
+
+	void Invalidate() { Time = FTimespan::MinValue(); }
+	bool IsValid() const { return Time != FTimespan::MinValue(); }
+
+	bool IsRelative() const { return SequenceIndex < 0; }
+
+	bool operator < (const FMediaTimeStamp & Other) const { return (SequenceIndex < Other.SequenceIndex) || (SequenceIndex == Other.SequenceIndex && Time < Other.Time); }
+	bool operator <= (const FMediaTimeStamp & Other) const { return (SequenceIndex < Other.SequenceIndex) || (SequenceIndex == Other.SequenceIndex && Time <= Other.Time); }
+	bool operator == (const FMediaTimeStamp & Other) const { return (SequenceIndex == Other.SequenceIndex && Time == Other.Time); }
+	bool operator > (const FMediaTimeStamp & Other) const { return (SequenceIndex > Other.SequenceIndex) || (SequenceIndex == Other.SequenceIndex && Time > Other.Time); }
+	bool operator >= (const FMediaTimeStamp & Other) const { return (SequenceIndex > Other.SequenceIndex) || (SequenceIndex == Other.SequenceIndex && Time >= Other.Time); }
+
+	FMediaTimeStamp operator + (const FTimespan & Other) const { return FMediaTimeStamp(Time + Other, SequenceIndex); }
+	FMediaTimeStamp operator - (const FTimespan & Other) const { return FMediaTimeStamp(Time - Other, SequenceIndex); }
+
+	FMediaTimeStamp operator - (const FMediaTimeStamp & Other) const { return FMediaTimeStamp(Time - Other.Time, -1); }
+
+	FTimespan Time;
+	int64 SequenceIndex;
+};
+
+class FMediaTimeStampSample
+{
+public:
+	FMediaTimeStampSample() : SampledAtTime(-1.0) {}
+	FMediaTimeStampSample(const FMediaTimeStamp & InTimeStamp, double InSampledAtTime) : TimeStamp(InTimeStamp), SampledAtTime(InSampledAtTime) {}
+
+	void Invalidate() { TimeStamp.Invalidate(); }
+	bool IsValid() const { return TimeStamp.IsValid(); }
+
+	FMediaTimeStamp TimeStamp;
+	double SampledAtTime;
+};
 
 /**
  * Interface for media clock time sources.
diff --git a/Engine/Source/Runtime/MediaAssets/Private/Assets/MediaPlayer.cpp b/Engine/Source/Runtime/MediaAssets/Private/Assets/MediaPlayer.cpp
--- a/Engine/Source/Runtime/MediaAssets/Private/Assets/MediaPlayer.cpp
+++ b/Engine/Source/Runtime/MediaAssets/Private/Assets/MediaPlayer.cpp
@@ -227,14 +227,16 @@ FTimespan UMediaPlayer::GetTime() const
 	return PlayerFacade->GetTime();
 }
 
-FTimespan UMediaPlayer::GetLastAudioSampleProcessedTime() const
+UMediaTimeStampInfo* UMediaPlayer::GetTimeStamp() const
 {
-	return PlayerFacade->GetLastAudioSampleProcessedTime();
-}
-
-FTimespan UMediaPlayer::GetLastVideoSampleProcessedTime() const
-{
-	return PlayerFacade->GetLastVideoSampleProcessedTime();
+	UMediaTimeStampInfo *TimeStampInfo = NewObject<UMediaTimeStampInfo>();
+	if (TimeStampInfo)
+	{
+		FMediaTimeStamp TimeStamp = PlayerFacade->GetTimeStamp();
+		TimeStampInfo->Time = TimeStamp.Time;
+		TimeStampInfo->SequenceIndex = TimeStamp.SequenceIndex;
+	}
+	return TimeStampInfo;
 }
 
 FText UMediaPlayer::GetTrackDisplayName(EMediaPlayerTrack TrackType, int32 TrackIndex) const
diff --git a/Engine/Source/Runtime/MediaAssets/Private/Assets/MediaSoundComponent.cpp b/Engine/Source/Runtime/MediaAssets/Private/Assets/MediaSoundComponent.cpp
--- a/Engine/Source/Runtime/MediaAssets/Private/Assets/MediaSoundComponent.cpp
+++ b/Engine/Source/Runtime/MediaAssets/Private/Assets/MediaSoundComponent.cpp
@@ -23,6 +23,9 @@ DECLARE_DWORD_COUNTER_STAT(TEXT("MediaUtils MediaSoundComponent Queued"), STAT_M
 
 CSV_DECLARE_CATEGORY_MODULE_EXTERN(MEDIA_API, MediaStreaming);
 
+
+static const int32 MaxAudioInputSamples = 8;	// accept at most these many samples into our input queue
+
 /* Static initialization
  *****************************************************************************/
 
@@ -127,7 +130,7 @@ void UMediaSoundComponent::UpdatePlayer()
 
 	// We have some audio decoders which are running with a limited amount of pre-allocated audio sample packets. 
 	// When the audio packets are not consumed in the UMediaSoundComponent::OnGenerateAudio method below, these packets are not 
-	// returned to the decooder which then cannot produce more audio samples. 
+	// returned to the decoder which then cannot produce more audio samples. 
 	//
 	// The UMediaSoundComponent::OnGenerateAudio is only called when our parent USynthComponent it active and
 	// this is conrolled by USynthComponent::Start() and USynthComponent::Stop(). We are tracking a state change here.
@@ -135,7 +138,7 @@ void UMediaSoundComponent::UpdatePlayer()
 	{
 		if (IsActive())
 		{
-			const auto NewSampleQueue = MakeShared<FMediaAudioSampleQueue, ESPMode::ThreadSafe>();
+			const auto NewSampleQueue = MakeShared<FMediaAudioSampleQueue, ESPMode::ThreadSafe>(MaxAudioInputSamples);
 			PlayerFacade->AddAudioSampleSink(NewSampleQueue);
 			{
 				FScopeLock Lock(&CriticalSection);
@@ -351,30 +354,37 @@ int32 UMediaSoundComponent::OnGenerateAudio(float* OutAudio, int32 NumSamples)
 		PinnedSampleQueue = SampleQueue;
 	}
 
+	// We have an input queue and are actively playing?
 	if (PinnedSampleQueue.IsValid() && (CachedRate != 0.0f))
 	{
 		const float Rate = CachedRate.Load();
 		const FTimespan Time = CachedTime.Load();
 
-		FTimespan OutTime = FTimespan::Zero();
-
 		{
 			const uint32 FramesRequested = uint32(NumSamples / NumChannels);
 			uint32 JumpFrame = MAX_uint32;
+			FMediaTimeStamp OutTime = FMediaTimeStamp(FTimespan::Zero());
 			uint32 FramesWritten = Resampler->Generate(OutAudio, OutTime, FramesRequested, Rate, Time, *PinnedSampleQueue, JumpFrame);
 			if (FramesWritten == 0)
 			{
 				return 0; // no samples available
 			}
 
+			// Fill in any gap left as we didn't have enough data
 			if (FramesWritten < FramesRequested)
 			{
 				memset(OutAudio + FramesWritten * NumChannels, 0, (NumSamples - FramesWritten * NumChannels) * sizeof(float));
 			}
-		}
 
-		LastPlaySampleTime = OutTime;
+			// Update audio time
+			LastPlaySampleTime = OutTime.Time;
+			PinnedSampleQueue->SetAudioTime(FMediaTimeStampSample(OutTime, FPlatformTime::Seconds()));
 
+			SET_FLOAT_STAT(STAT_MediaUtils_MediaSoundComponentSampleTime, OutTime.Time.GetTotalSeconds());
+			SET_DWORD_STAT(STAT_Media_SoundCompQueued, PinnedSampleQueue->Num());
+		}
+
+		// - - - - - - - - - - - - - - - - - - - - - - - - - - - - - - - - - - - - - - - - - - - - - - - - - - - - - - - - - - - - - - - - - - - -
 
 		if (bSpectralAnalysisEnabled || bEnvelopeFollowingEnabled)
 		{
@@ -427,10 +437,6 @@ int32 UMediaSoundComponent::OnGenerateAudio(float* OutAudio, int32 NumSamples)
 				}
 			}
 		}
-
-		SET_FLOAT_STAT(STAT_MediaUtils_MediaSoundComponentSync, FMath::Abs((Time - OutTime).GetTotalMilliseconds()));
-		SET_FLOAT_STAT(STAT_MediaUtils_MediaSoundComponentSampleTime, OutTime.GetTotalMilliseconds());
-		SET_DWORD_STAT(STAT_Media_SoundCompQueued, PinnedSampleQueue->Num());
 	}
 	else
 	{
diff --git a/Engine/Source/Runtime/MediaAssets/Private/Assets/MediaTexture.cpp b/Engine/Source/Runtime/MediaAssets/Private/Assets/MediaTexture.cpp
--- a/Engine/Source/Runtime/MediaAssets/Private/Assets/MediaTexture.cpp
+++ b/Engine/Source/Runtime/MediaAssets/Private/Assets/MediaTexture.cpp
@@ -14,6 +14,7 @@
 #include "UObject/WeakObjectPtrTemplates.h"
 
 #include "MediaPlayer.h"
+#include "IMediaPlayer.h"
 #include "Misc/MediaTextureResource.h"
 #include "IMediaTextureSample.h"
 
@@ -62,12 +63,15 @@ UMediaTexture::UMediaTexture(const FObjectInitializer& ObjectInitializer)
 	, ClearColor(FLinearColor::Black)
 	, EnableGenMips(false)
 	, NumMips(1)
+	, NewStyleOutput(false)
+	, OutputFormat(MTOF_Default)
 	, DefaultGuid(FGuid::NewGuid())
 	, Dimensions(FIntPoint::ZeroValue)
 	, Size(0)
 	, CachedNextSampleTime(FTimespan::MinValue())
 {
 	NeverStream = true;
+	SRGB = true;
 }
 
 
@@ -150,6 +154,10 @@ FTextureResource* UMediaTexture::CreateResource()
 
 EMaterialValueType UMediaTexture::GetMaterialType() const
 {
+	if (NewStyleOutput)
+	{
+		return MCT_Texture2D;
+	}
 	return EnableGenMips ? MCT_Texture2D : MCT_TextureExternal;
 }
 
@@ -330,15 +338,55 @@ void UMediaTexture::TickResource(FTimespan Timecode)
 
 		if (PlayerActive)
 		{
-			RenderParams.SampleSource = SampleQueue;
+			check(CurrentPlayerPtr->GetPlayerFacade()->GetPlayer());
+
+			if (CurrentPlayerPtr->GetPlayerFacade()->GetPlayer()->GetPlayerFeatureFlag(IMediaPlayer::EFeatureFlag::UsePlaybackTimingV2))
+			{
+				/*
+					We are using the old-style "1sample queue to sink" architecture to actually just pass long only ONE sample at a time from the logic
+					inside the player facade to the sinks. The selection as to what to render this frame is expected to be done earlier
+					this frame on the gamethread, hence only a single output frame is selected and passed along...
+				*/
+				TSharedPtr<IMediaTextureSample, ESPMode::ThreadSafe> Sample;
+				while (SampleQueue->Dequeue(Sample))
+					;
+
+				if (!Sample.IsValid())
+				{
+					// Player is active (do not clear), but we have no new data
+					// -> we do not need to trigger anything on the renderthread
+					return;
+				}
+
+				RenderParams.TextureSample = Sample;
+
+				RenderParams.Rate = CurrentPlayerPtr->GetRate();
+				RenderParams.Time = Sample->GetTime();
+
+				if (NewStyleOutput)
+				{
+					// For new-style output the sample's sRGB state controls what we output
+					// (FOR NOW: this is too simplified if we have more then Rec703 material)
+					SRGB = Sample->IsOutputSrgb();
+					// Ensure sRGB changes will not trigger rendering the next time around
+					LastSrgb = SRGB;
+				}
+			}
+			else
+			{
+				//
+				// Old style: pass queue along and dequeue only at render time
+				//
+				RenderParams.SampleSource = SampleQueue;
+
+				RenderParams.Rate = CurrentPlayerPtr->GetRate();
+				RenderParams.Time = CurrentPlayerPtr->GetTime();
+			}
 		}
 		else if (!AutoClear)
 		{
 			return; // retain last frame
 		}
-
-		RenderParams.Rate = CurrentPlayerPtr->GetRate();
-		RenderParams.Time = CurrentPlayerPtr->GetTime();
 	}
 	else if (!AutoClear && (CurrentGuid == PreviousGuid))
 	{
diff --git a/Engine/Source/Runtime/MediaAssets/Private/Misc/MediaTextureResource.cpp b/Engine/Source/Runtime/MediaAssets/Private/Misc/MediaTextureResource.cpp
--- a/Engine/Source/Runtime/MediaAssets/Private/Misc/MediaTextureResource.cpp
+++ b/Engine/Source/Runtime/MediaAssets/Private/Misc/MediaTextureResource.cpp
@@ -66,7 +66,7 @@ namespace MediaTextureResourceHelpers
 
 		case EMediaTextureSampleFormat::CharNV12:
 		case EMediaTextureSampleFormat::CharNV21:
-			return PF_G8;
+			return PF_G8; // note: right now this case will be encountered only if CPU-side data in NV12/21 format is in sample -> in this case we cannot create a true NV12 texture OR the platforms view it as U8s anyway 
 
 		case EMediaTextureSampleFormat::FloatRGB:
 			return PF_FloatRGB;
@@ -103,44 +103,6 @@ namespace MediaTextureResourceHelpers
 		}
 	}
 
-
-	/**
-	 * Check whether the given sample requires a conversion shader.
-	 *
-	 * @param Sample The sample to check.
-	 * @param SrgbOutput Whether the output is expected in sRGB color space.
-	 * @return true if conversion is required, false otherwise.
-	 * @see RequiresOutputResource
-	 */
-	bool RequiresConversion(const TSharedPtr<IMediaTextureSample, ESPMode::ThreadSafe>& Sample, bool SrgbOutput)
-	{
-		// If the output color space is expected to be sRGB, but the
-		// sample is not, a color space conversion on the GPU is required.
-
-		if (Sample->IsOutputSrgb() != SrgbOutput)
-		{
-			return true;
-		}
-
-		// If the output dimensions are not the same as the sample's
-		// dimensions, a resizing conversion on the GPU is required.
-
-		if (Sample->GetDim() != Sample->GetOutputDim())
-		{
-			return true;
-		}
-
-		// Only the following pixel formats are supported natively.
-		// All other formats require a conversion on the GPU.
-
-		const EMediaTextureSampleFormat Format = Sample->GetFormat();
-
-		return ((Format != EMediaTextureSampleFormat::CharBGRA) &&
-				(Format != EMediaTextureSampleFormat::FloatRGB) &&
-				(Format != EMediaTextureSampleFormat::FloatRGBA));
-	}
-
-
 	/**
 	 * Check whether the given sample requires an sRGB texture.
 	 *
@@ -179,9 +141,9 @@ FMediaTextureResource::FMediaTextureResource(UMediaTexture& InOwner, FIntPoint&
 	, CurrentSamplerFilter(ESamplerFilter_Num)
 {
 #if PLATFORM_ANDROID
-	bUsesImageExternal = (!FAndroidMisc::ShouldUseVulkan() && GSupportsImageExternal);
+	bUsesImageExternal = !Owner.NewStyleOutput && (!FAndroidMisc::ShouldUseVulkan() && GSupportsImageExternal);
 #else
-	bUsesImageExternal = GSupportsImageExternal;
+	bUsesImageExternal = !Owner.NewStyleOutput && GSupportsImageExternal;
 #endif
 }
 
@@ -202,16 +164,23 @@ void FMediaTextureResource::Render(const FRenderParams& Params)
 
 	TSharedPtr<FMediaTextureSampleSource, ESPMode::ThreadSafe> SampleSource = Params.SampleSource.Pin();
 
+	// Do we either have a classic sample source (queue) or a single, explicit sample to display?
+	if (SampleSource.IsValid() || Params.TextureSample.IsValid())
+	{
+		TSharedPtr<IMediaTextureSample, ESPMode::ThreadSafe> Sample;
+		bool UseSample;
+
+		// Yes, it it a queue?
 	if (SampleSource.IsValid())
 	{
+			// Yes, find out what we will display...
+			UseSample = false;
+
 		// get the most current sample to be rendered
 		TSharedPtr<IMediaTextureSample, ESPMode::ThreadSafe> TestSample;
-		TSharedPtr<IMediaTextureSample, ESPMode::ThreadSafe> Sample;
-		bool UseSample = false;
-		
 		while (SampleSource->Peek(TestSample) && TestSample.IsValid())
 		{
-			const FTimespan StartTime = TestSample->GetTime();
+				const FTimespan StartTime = TestSample->GetTime().Time;
 			const FTimespan EndTime = StartTime + TestSample->GetDuration();
 
 			if ((Params.Rate >= 0.0f) && (Params.Time < StartTime))
@@ -232,6 +201,15 @@ void FMediaTextureResource::Render(const FRenderParams& Params)
 #endif
 
 			UseSample = SampleSource->Dequeue(Sample);
+			}
+		}
+		else
+		{
+			// We do have an explicit sample to display...
+			// (or nothing)
+			Sample = Params.TextureSample;
+			UseSample = Sample.IsValid();
+		}
 
 #if MEDIATEXTURERESOURCE_TRACE_RENDER
 			if (!UseSample && Sample.IsValid())
@@ -243,7 +221,6 @@ void FMediaTextureResource::Render(const FRenderParams& Params)
 				);
 			}
 #endif
-		}
 
 #if UE_MEDIAUTILS_DEVELOPMENT_DELEGATE
 		FMediaDelegates::OnPreSampleRender_RenderThread.Broadcast(&Owner, UseSample, Sample);
@@ -255,12 +232,15 @@ void FMediaTextureResource::Render(const FRenderParams& Params)
 
 		// If real "external texture" support is in place and no mips are used the image will "bypass" any of this processing via the GUID-based lookup for "ExternalTextures" and will
 		// reach the reading material shader without any processing here...
+		// (note that if "new style output" is enabled we will ALWAYS see bUsesImageExternal as FALSE)
 		if (UseSample && !(bUsesImageExternal && !bEnableGenMips))
 		{
 			//
 			// Valid sample & Sample should be shown
 			//
 
+			bool ConvertOrCopyNeeded = false;
+
 			if (Sample->GetOutputDim().GetMin() <= 0)
 			{
 				//
@@ -268,16 +248,66 @@ void FMediaTextureResource::Render(const FRenderParams& Params)
 				//
 				ClearTexture(FLinearColor::Red, Params.SrgbOutput); // mark corrupt sample
 			}
-			else if (Sample->GetMediaTextureSampleConverter())
+			else if (IMediaTextureSampleConverter *Converter = Sample->GetMediaTextureSampleConverter())
 			{
 				//
 				// Sample brings its own converter
 				//
+
+				IMediaTextureSampleConverter::FConversionHints Hints;
+				Hints.bOutputSRGB = Params.SrgbOutput;
+				Hints.NumMips = NumMips;
+
+				// Does the conversion create its own output texture?
+				if ((Converter->GetConverterInfoFlags() & IMediaTextureSampleConverter::ConverterInfoFlags_WillCreateOutputTexture) == 0)
+				{
+					// No. Does it actually do the conversion or just a pre-process step not yielding real output?
+					if (Converter->GetConverterInfoFlags() & IMediaTextureSampleConverter::ConverterInfoFlags_PreprocessOnly)
+					{
+						// Preprocess...
+						FTexture2DRHIRef DummyTexture;
+						if (Converter->Convert(DummyTexture, Hints))
+						{
+							// ...followed by the built in conversion code as needed...
+							ConvertOrCopyNeeded = true;
+						}
+					}
+					else
+					{
+						// Conversion is fully handled by converter
 				CreateOutputRenderTarget(Sample->GetOutputDim(), MediaTextureResourceHelpers::GetConvertedPixelFormat(Sample), Params.SrgbOutput, Params.ClearColor, NumMips);
-				Sample->GetMediaTextureSampleConverter()->Convert(RenderTargetTextureRHI);
+						Converter->Convert(RenderTargetTextureRHI, Hints);
+					}
+				}
+				else
+				{
+					// The converter will create its own output texture (not preferred, but sometimes the only way)
+					FTexture2DRHIRef OutTexture;
+					if (Converter->Convert(OutTexture, Hints))
+					{
+						// As the converter created the texture, we might need to convert it even more to make it fit our needs. Check...
+						if (RequiresConversion(OutTexture, Sample->GetOutputDim(), Params.SrgbOutput, NumMips))
+						{
+							CreateOutputRenderTarget(Sample->GetOutputDim(), MediaTextureResourceHelpers::GetConvertedPixelFormat(Sample), Params.SrgbOutput, Params.ClearColor, NumMips);
+							ConvertTextureToOutput(OutTexture.GetReference(), Sample, Params.SrgbOutput);
+						}
+						else
+						{
+							UpdateTextureReference(OutTexture);
+						}
+					}
+				}
+
 				Cleared = false;
 			}
-			else if (MediaTextureResourceHelpers::RequiresConversion(Sample, Params.SrgbOutput))
+			else
+			{
+				ConvertOrCopyNeeded = true;
+			}
+
+			if (ConvertOrCopyNeeded)
+			{
+				if (RequiresConversion(Sample, Params.SrgbOutput, NumMips))
 			{
 				//
 				// Sample needs to be converted by built in converter code
@@ -291,6 +321,7 @@ void FMediaTextureResource::Render(const FRenderParams& Params)
 				//
 				CopySample(Sample, Params.ClearColor, Params.SrgbOutput, NumMips, Params.CurrentGuid);
 			}
+			}
 
 			Rotation = Sample->GetScaleRotation();
 			Offset = Sample->GetOffset();
@@ -317,7 +348,7 @@ void FMediaTextureResource::Render(const FRenderParams& Params)
 				RHICmdList.TransitionResource(EResourceTransitionAccess::EReadable, OutputTarget);
 			}
 
-			SET_FLOAT_STAT(STAT_MediaUtils_TextureSampleTime, Sample->GetTime().GetTotalMilliseconds());
+			SET_FLOAT_STAT(STAT_MediaUtils_TextureSampleTime, Sample->GetTime().Time.GetTotalMilliseconds());
 		}
 		else
 		{
@@ -362,7 +393,7 @@ void FMediaTextureResource::Render(const FRenderParams& Params)
 
 	// Update external texture registration in case we have no native support
 	// (in that case there is support, the player will do this - but it is used all the time)
-	if (!bUsesImageExternal)
+	if (!Owner.NewStyleOutput && !bUsesImageExternal)
 	{
 		SetupSampler();
 
@@ -475,7 +506,7 @@ void FMediaTextureResource::ReleaseDynamicRHI()
 /* FMediaTextureResource implementation
  *****************************************************************************/
 
-void FMediaTextureResource::ClearTexture(const FLinearColor& ClearColor, bool SrgbOutput)
+ void FMediaTextureResource::ClearTexture(const FLinearColor& ClearColor, bool SrgbOutput)
 {
 	// create output render target if we don't have one yet
 	CreateOutputRenderTarget(FIntPoint(2, 2), PF_B8G8R8A8, SrgbOutput, ClearColor, 1);
@@ -496,8 +527,111 @@ void FMediaTextureResource::ClearTexture(const FLinearColor& ClearColor, bool Sr
 }
 
 
-void FMediaTextureResource::ConvertSample(const TSharedPtr<IMediaTextureSample, ESPMode::ThreadSafe>& Sample, const FLinearColor& ClearColor, bool SrgbOutput, uint8 InNumMips)
-{
+ bool FMediaTextureResource::RequiresConversion(const FTexture2DRHIRef& SampleTexture, const FIntPoint & OutputDim, bool SrgbOutput, uint8 InNumMips) const
+ {
+	 if (!Owner.NewStyleOutput)
+	 {
+		 //
+		 // Old Style
+		 //
+
+		 if (((SampleTexture->GetFlags() & TexCreate_SRGB) != 0) != SrgbOutput)
+		 {
+			 return true;
+		 }
+	 }
+	 else
+	 {
+		 //
+		 // New Style
+		 // 
+
+		 // For now we only allow this, single SRGB-style output format
+		 check(Owner.OutputFormat == MTOF_SRGB_LINOUT || Owner.OutputFormat == MTOF_Default);
+
+		 if ((SampleTexture->GetFlags() & TexCreate_SRGB) == 0)
+		 {
+			 return true;
+		 }
+
+		 if (SampleTexture->GetNumMips() != InNumMips)
+		 {
+			 return true;
+		 }
+	 }
+
+	 if (SampleTexture->GetSizeXY() != OutputDim)
+	 {
+		 return true;
+	 }
+
+	 // Only the following pixel formats are supported natively.
+	 // All other formats require a conversion on the GPU.
+
+	 const EPixelFormat Format = SampleTexture->GetFormat();
+
+	 return ((Format != PF_B8G8R8A8) &&
+		 (Format != PF_FloatRGB) &&
+		 (Format != PF_FloatRGBA));
+ }
+
+
+ bool FMediaTextureResource::RequiresConversion(const TSharedPtr<IMediaTextureSample, ESPMode::ThreadSafe>& Sample, bool SrgbOutput, uint8 InNumMips) const
+ {
+	 if (!Owner.NewStyleOutput)
+	 {
+		 //
+		 // Old Style
+		 //
+
+		 if (Sample->IsOutputSrgb() != SrgbOutput)
+		 {
+			 return true;
+		 }
+	 }
+	 else
+	 {
+		 //
+		 // New Style
+		 // 
+
+		 // For now we only allow this, single SRGB-style output format
+		 check(Owner.OutputFormat == MTOF_SRGB_LINOUT || Owner.OutputFormat == MTOF_Default);
+
+		 if (!Sample->IsOutputSrgb())
+		 {
+			 return true;
+		 }
+
+		 FRHITexture *Texture = Sample->GetTexture();
+
+		 if (Texture && Texture->GetNumMips() != InNumMips)
+		 {
+			 return true;
+		 }
+	 }
+
+	 // If the output dimensions are not the same as the sample's
+	 // dimensions, a resizing conversion on the GPU is required.
+
+	 if (Sample->GetDim() != Sample->GetOutputDim())
+	 {
+		 return true;
+	 }
+
+	 // Only the following pixel formats are supported natively.
+	 // All other formats require a conversion on the GPU.
+
+	 const EMediaTextureSampleFormat Format = Sample->GetFormat();
+
+	 return ((Format != EMediaTextureSampleFormat::CharBGRA) &&
+		 (Format != EMediaTextureSampleFormat::FloatRGB) &&
+		 (Format != EMediaTextureSampleFormat::FloatRGBA));
+ }
+
+
+ void FMediaTextureResource::ConvertSample(const TSharedPtr<IMediaTextureSample, ESPMode::ThreadSafe>& Sample, const FLinearColor& ClearColor, bool SrgbOutput, uint8 InNumMips)
+ {
 	const EPixelFormat InputPixelFormat = MediaTextureResourceHelpers::GetPixelFormat(Sample);
 
 	// get input texture
@@ -510,7 +644,7 @@ void FMediaTextureResource::ConvertSample(const TSharedPtr<IMediaTextureSample,
 		FRHITexture* SampleTexture = Sample->GetTexture();
 		FRHITexture2D* SampleTexture2D = (SampleTexture != nullptr) ? SampleTexture->GetTexture2D() : nullptr;
 
-		if (SampleTexture2D != nullptr)
+		 if (SampleTexture2D)
 		{
 			// Use the sample as source texture...
 
@@ -524,27 +658,21 @@ void FMediaTextureResource::ConvertSample(const TSharedPtr<IMediaTextureSample,
 			// Make a source texture so we can convert from it...
 
 			const bool SrgbTexture = MediaTextureResourceHelpers::RequiresSrgbTexture(Sample);
-			const uint32 InputCreateFlags = TexCreate_Dynamic | (SrgbTexture ? TexCreate_SRGB : 0);
+			 const uint32 InputCreateFlags = TexCreate_Dynamic | (SrgbTexture ? TexCreate_SRGB : 0);  //<<< FILTER SRGB AWAY FOR G8 AND SUCH? (or will the RHI code kindly do this?)
 			const FIntPoint SampleDim = Sample->GetDim();
 
 			// create a new temp input render target if necessary
 			if (!InputTarget.IsValid() || (InputTarget->GetSizeXY() != SampleDim) || (InputTarget->GetFormat() != InputPixelFormat) || ((InputTarget->GetFlags() & InputCreateFlags) != InputCreateFlags))
 			{
-				TRefCountPtr<FRHITexture2D> DummyTexture2DRHI;
 				FRHIResourceCreateInfo CreateInfo;
-
-				RHICreateTargetableShaderResource2D(
+				 InputTarget = RHICreateTexture2D(
 					SampleDim.X,
 					SampleDim.Y,
 					InputPixelFormat,
 					1,
+					 1,
 					InputCreateFlags,
-					TexCreate_RenderTargetable,
-					false,
-					CreateInfo,
-					InputTarget,
-					DummyTexture2DRHI
-				);
+					 CreateInfo);
 
 				UpdateResourceSize();
 			}
@@ -561,6 +689,12 @@ void FMediaTextureResource::ConvertSample(const TSharedPtr<IMediaTextureSample,
 	const FIntPoint OutputDim = Sample->GetOutputDim();
 	CreateOutputRenderTarget(OutputDim, MediaTextureResourceHelpers::GetConvertedPixelFormat(Sample), SrgbOutput, ClearColor, InNumMips);
 
+	 ConvertTextureToOutput(InputTexture, Sample, SrgbOutput);
+ }
+
+
+ void FMediaTextureResource::ConvertTextureToOutput(FRHITexture2D* InputTexture, const TSharedPtr<IMediaTextureSample, ESPMode::ThreadSafe>& Sample, bool SrgbOutput)
+ {
 	// perform the conversion
 	FRHICommandListImmediate& CommandList = FRHICommandListExecutor::GetImmediateCommandList();
 	{
@@ -570,6 +704,8 @@ void FMediaTextureResource::ConvertSample(const TSharedPtr<IMediaTextureSample,
 		FGraphicsPipelineStateInitializer GraphicsPSOInit;
 		FRHITexture* RenderTarget = RenderTargetTextureRHI.GetReference();
 
+		FIntPoint OutputDim(RenderTarget->GetSizeXYZ().X, RenderTarget->GetSizeXYZ().Y);
+
 		FRHIRenderPassInfo RPInfo(RenderTarget, ERenderTargetActions::DontLoad_Store);
 		CommandList.BeginRenderPass(RPInfo, TEXT("ConvertMedia"));
 		{
@@ -598,6 +734,7 @@ void FMediaTextureResource::ConvertSample(const TSharedPtr<IMediaTextureSample,
 
 			bool bIsSampleOutputSrgb = Sample->IsOutputSrgb();
 
+			// Use the sample format to choose the conversion path
 			switch (Sample->GetFormat())
 			{
 			case EMediaTextureSampleFormat::CharAYUV:
@@ -620,15 +757,26 @@ void FMediaTextureResource::ConvertSample(const TSharedPtr<IMediaTextureSample,
 
 			case EMediaTextureSampleFormat::CharNV12:
 			{
+				if (InputTexture->GetFormat() == PF_NV12)
+				{
 				TShaderMapRef<FNV12ConvertPS> ConvertShader(ShaderMap);
 				GraphicsPSOInit.BoundShaderState.PixelShaderRHI = ConvertShader.GetPixelShader();
 				SetGraphicsPipelineState(CommandList, GraphicsPSOInit);
 				ConvertShader->SetParameters(CommandList, InputTexture, OutputDim, YUVToRGBMatrix, YUVOffset, bIsSampleOutputSrgb);
+				}
+			else
+				{
+					TShaderMapRef<FNV12ConvertAsBytesPS> ConvertShader(ShaderMap);
+					GraphicsPSOInit.BoundShaderState.PixelShaderRHI = ConvertShader.GetPixelShader();
+					SetGraphicsPipelineState(CommandList, GraphicsPSOInit);
+					ConvertShader->SetParameters(CommandList, InputTexture, OutputDim, YUVToRGBMatrix, YUVOffset, bIsSampleOutputSrgb);
+				}
 			}
 			break;
 
 			case EMediaTextureSampleFormat::CharNV21:
 			{
+				// source texture might be NV12 or G8...
 				TShaderMapRef<FNV21ConvertPS> ConvertShader(ShaderMap);
 				GraphicsPSOInit.BoundShaderState.PixelShaderRHI = ConvertShader.GetPixelShader();
 				SetGraphicsPipelineState(CommandList, GraphicsPSOInit);
@@ -685,6 +833,7 @@ void FMediaTextureResource::ConvertSample(const TSharedPtr<IMediaTextureSample,
 			case EMediaTextureSampleFormat::FloatRGB:
 			case EMediaTextureSampleFormat::FloatRGBA:
 			{
+				// Simple 1:1 copy (but using normal texture sampler: sRGB conversions may occur depending on setup; any manuao sRGB/linear conversion is disabled)
 				TShaderMapRef<FRGBConvertPS> ConvertShader(ShaderMap);
 				GraphicsPSOInit.BoundShaderState.PixelShaderRHI = ConvertShader.GetPixelShader();
 				SetGraphicsPipelineState(CommandList, GraphicsPSOInit);
@@ -745,9 +894,10 @@ void FMediaTextureResource::CopySample(const TSharedPtr<IMediaTextureSample, ESP
 		// Texture to receive precisely only output pixels via CPU copy
 		CreateOutputRenderTarget(Sample->GetDim(), MediaTextureResourceHelpers::GetPixelFormat(Sample), SrgbOutput, ClearColor, InNumMips);
 
-		if (bEnableGenMips && bUsesImageExternal && !Sample->GetBuffer())
+		// If we also have no source buffer and the platform generally would allow for use of external textures, we assume it is just that...
+		// (as long as the player actually produces (dummy) samples, this will enable mips support as well as auto conversion for "new style output" mode)
+		if (!Sample->GetBuffer() && GSupportsImageExternal)
 		{
-			// We are using an external texture to read from
 			CopyFromExternalTexture(Sample, TextureGUID);
 		}
 		else
@@ -923,7 +1073,7 @@ void FMediaTextureResource::CacheNextAvailableSampleTime(const TSharedPtr<FMedia
 		TSharedPtr<IMediaTextureSample, ESPMode::ThreadSafe> Sample;
 		if (InSampleQueue->Peek(Sample))
 		{
-			SampleTime = Sample->GetTime();
+			SampleTime = Sample->GetTime().Time;
 		}
 	}
 
diff --git a/Engine/Source/Runtime/MediaAssets/Private/Misc/MediaTextureResource.h b/Engine/Source/Runtime/MediaAssets/Private/Misc/MediaTextureResource.h
--- a/Engine/Source/Runtime/MediaAssets/Private/Misc/MediaTextureResource.h
+++ b/Engine/Source/Runtime/MediaAssets/Private/Misc/MediaTextureResource.h
@@ -12,6 +12,7 @@
 #include "Templates/SharedPointer.h"
 #include "TextureResource.h"
 #include "UnrealClient.h"
+#include "IMediaTimeSource.h"
 
 class FMediaPlayerFacade;
 class IMediaPlayer;
@@ -78,7 +79,10 @@ public:
 		uint8 NumMips;
 
 		/** The time of the video frame to render (in player's clock). */
-		FTimespan Time;
+		FMediaTimeStamp Time;
+
+		/** Explicit texture sample to render - if set time will be ignored */
+		TSharedPtr<IMediaTextureSample, ESPMode::ThreadSafe> TextureSample;
 	};
 
 	/**
@@ -128,6 +132,8 @@ protected:
 	 */
 	void ConvertSample(const TSharedPtr<IMediaTextureSample, ESPMode::ThreadSafe>& Sample, const FLinearColor& ClearColor, bool SrgbOutput, uint8 InNumMips);
 
+	void ConvertTextureToOutput(FRHITexture2D* InputTexture, const TSharedPtr<IMediaTextureSample, ESPMode::ThreadSafe>& Sample, bool SrgbOutput);
+
 	/**
 	 * Render the given texture sample by using it as or copying it to the render target.
 	 *
@@ -167,6 +173,9 @@ protected:
 	/** Copy to local buffer from external texture */
 	void CopyFromExternalTexture(const TSharedPtr <IMediaTextureSample, ESPMode::ThreadSafe>& Sample, const FGuid & TextureGUID);
 
+	bool RequiresConversion(const TSharedPtr<IMediaTextureSample, ESPMode::ThreadSafe>& Sample, bool SrgbOutput, uint8 numMips) const;
+	bool RequiresConversion(const FTexture2DRHIRef& SampleTexture, const FIntPoint & OutputDim, bool SrgbOutput, uint8 numMips) const;
+
 private:
 
 	// Class to maintain a reference to a IMediaTextureSample instance as long  as needed by RHI
diff --git a/Engine/Source/Runtime/MediaAssets/Public/MediaPlayer.h b/Engine/Source/Runtime/MediaAssets/Public/MediaPlayer.h
--- a/Engine/Source/Runtime/MediaAssets/Public/MediaPlayer.h
+++ b/Engine/Source/Runtime/MediaAssets/Public/MediaPlayer.h
@@ -15,6 +15,7 @@
 #include "Misc/Guid.h"
 #include "Engine/LatentActionManager.h"
 #include "MediaPlayerOptions.h"
+#include "IMediaTimeSource.h"
 
 #include "MediaPlayer.generated.h"
 
@@ -67,6 +68,23 @@ enum class EMediaPlayerTrack : uint8
 };
 
 
+UCLASS(BlueprintType, hidecategories = (Object))
+class MEDIAASSETS_API UMediaTimeStampInfo
+	: public UObject
+{
+	GENERATED_BODY()
+
+public:
+	UMediaTimeStampInfo() : Time(FTimespan::Zero()), SequenceIndex(0) {}
+
+	UPROPERTY(BlueprintReadOnly, Category = "Media|Time")
+	FTimespan Time;
+
+	UPROPERTY(BlueprintReadOnly, Category = "Media|Time")
+	int64 SequenceIndex;
+};
+
+
 /**
  * Implements a media player asset that can play movies and other media sources.
  */
@@ -287,20 +305,13 @@ public:
 	FTimespan GetTime() const;
 
 	/**
-	 * Get time of last audio sample decoded
+	 * Get the media's current playback timestamp.
 	 *
-	 * @return Time of last audio sample decoded.
-	 */
-	UFUNCTION(BlueprintCallable, Category="Media|MediaPlayer")
-	FTimespan GetLastAudioSampleProcessedTime() const;
-
-	/**
-	 * Get time of last video sample decoded
-	 *
-	 * @return Time of last video sample decoded.
+	 * @return Playback timestamp.
+	 * @see GetDuration, Seek
 	 */
-	UFUNCTION(BlueprintCallable, Category="Media|MediaPlayer")
-	FTimespan GetLastVideoSampleProcessedTime() const;
+	UFUNCTION(BlueprintCallable, Category = "Media|MediaPlayer")
+	UMediaTimeStampInfo* GetTimeStamp() const;
 
 	/**
 	 * Get the human readable name of the specified track.
diff --git a/Engine/Source/Runtime/MediaAssets/Public/MediaSoundComponent.h b/Engine/Source/Runtime/MediaAssets/Public/MediaSoundComponent.h
--- a/Engine/Source/Runtime/MediaAssets/Public/MediaSoundComponent.h
+++ b/Engine/Source/Runtime/MediaAssets/Public/MediaSoundComponent.h
@@ -274,10 +274,10 @@ private:
 	/** Which frequencies to analyze. */
 	TArray<float> FrequenciesToAnalyze;
 
-	/** The FFT bin-size to use for FFT analysis. Smaller sizes make it more reactive but less acurrate in the frequency space. */
+	/** The FFT bin-size to use for FFT analysis. Smaller sizes make it more reactive but less accurate in the frequency space. */
 	EMediaSoundComponentFFTSize FFTSize;
 
-	/** Spectrum analyzer used for anlayzing audio in media. */
+	/** Spectrum analyzer used for analyzing audio in media. */
 	Audio::FSpectrumAnalyzer SpectrumAnalyzer;
 	Audio::FSpectrumAnalyzerSettings SpectrumAnalyzerSettings;
 
diff --git a/Engine/Source/Runtime/MediaAssets/Public/MediaTexture.h b/Engine/Source/Runtime/MediaAssets/Public/MediaTexture.h
--- a/Engine/Source/Runtime/MediaAssets/Public/MediaTexture.h
+++ b/Engine/Source/Runtime/MediaAssets/Public/MediaTexture.h
@@ -22,6 +22,14 @@ class IMediaTextureSample;
 class UMediaPlayer;
 
 
+UENUM()
+enum MediaTextureOutputFormat
+{
+	MTOF_Default					UMETA(DisplayName = "Default (sRGB)"),
+	MTOF_SRGB_LINOUT				UMETA(DisplayName = "sRGB (linear output)"),		// sRGB data, using sRGB texture formats; hence read as linear RGB
+	MTOF_MAX,
+};
+
 /**
  * Implements a texture asset for rendering video tracks from UMediaPlayer assets.
  */
@@ -47,7 +55,7 @@ class MEDIAASSETS_API UMediaTexture
 	UPROPERTY(EditAnywhere, BlueprintReadWrite, Category="MediaTexture")
 	FLinearColor ClearColor;
 
-	/** Basic enablement for mip generation. */
+	/** Basic enablement for mip generation (default = false). */
 	UPROPERTY(EditAnywhere, BlueprintReadOnly, Category = "MediaTexture", meta = (DisplayName = "Enable Mipmap generation"))
 	bool EnableGenMips;
 
@@ -55,6 +63,13 @@ class MEDIAASSETS_API UMediaTexture
 	UPROPERTY(EditAnywhere, BlueprintReadWrite, Category = "MediaTexture", meta=(DisplayName="Total number of Mipmaps to output"))
 	uint8 NumMips;
 
+	/** Enable new style output (default = false). */
+	UPROPERTY(EditAnywhere, BlueprintReadOnly, Category = "MediaTexture", meta = (DisplayName = "Enable new style output"))
+	bool NewStyleOutput;
+
+	UPROPERTY(EditAnywhere, BlueprintReadWrite, Category = "MediaTexture", meta = (DisplayName = "Output format (new style)"))
+	TEnumAsByte<enum MediaTextureOutputFormat> OutputFormat;
+
 public:
 
 	/**
diff --git a/Engine/Source/Runtime/MediaUtils/Private/MediaAudioResampler.cpp b/Engine/Source/Runtime/MediaUtils/Private/MediaAudioResampler.cpp
--- a/Engine/Source/Runtime/MediaUtils/Private/MediaAudioResampler.cpp
+++ b/Engine/Source/Runtime/MediaUtils/Private/MediaAudioResampler.cpp
@@ -334,7 +334,7 @@ void FMediaAudioResampler::Flush()
 }
 
 
-uint32 FMediaAudioResampler::Generate(float* Output, FTimespan& OutTime, const uint32 FramesRequested, float Rate, FTimespan Time, FMediaAudioSampleSource& SampleSource, uint32& JumpFrame)
+uint32 FMediaAudioResampler::Generate(float* Output, FMediaTimeStamp& OutTime, const uint32 FramesRequested, float Rate, FTimespan Time, FMediaAudioSampleSource& SampleSource, uint32& JumpFrame)
 {
 	JumpFrame = MAX_uint32;
 	if ((FramesRequested == 0) || (Output == nullptr) || (OutputSampleRate == 0))
@@ -351,10 +351,10 @@ uint32 FMediaAudioResampler::Generate(float* Output, FTimespan& OutTime, const u
 		{
 			double RateDirection = Rate > 0 ? 1 : -1;
 			TRange<FTimespan> ContiguousRange = TRange<FTimespan>::Inclusive(
-				InputTime + InputDuration * 0.5 * RateDirection,
-				InputTime + InputDuration * 1.5 * RateDirection);
+				InputTime.Time + InputDuration * 0.5 * RateDirection,
+				InputTime.Time + InputDuration * 1.5 * RateDirection);
 
-			FTimespan PreviousInputTime = InputTime;
+			FMediaTimeStamp PreviousInputTime = InputTime;
 			FTimespan PreviousInputDuration = InputDuration;
 
 			// fetch next sample
@@ -370,12 +370,12 @@ uint32 FMediaAudioResampler::Generate(float* Output, FTimespan& OutTime, const u
 			}
 
 			// Detect jumps in time (looping, fast forwarding, seeking)
-			bool bIsContiguous = ContiguousRange.Contains(InputTime);
+			bool bIsContiguous = ContiguousRange.Contains(InputTime.Time);
 			if (!bIsContiguous)
 			{
 				JumpFrame = FramesGenerated;
 			}
-			UE_LOG(LogMediaUtils, VeryVerbose, TEXT("[AUDIO DEQUEUE] PlayerTime: %s, SampleTime: %s, Contiguous: %d, Frames: %d"), *Time.ToString(), *InputTime.ToString(), bIsContiguous, NextSample->GetFrames());
+			UE_LOG(LogMediaUtils, VeryVerbose, TEXT("[AUDIO DEQUEUE] PlayerTime: %s, SampleTime: %s, Contiguous: %d, Frames: %d"), *Time.ToString(), *InputTime.Time.ToString(), bIsContiguous, NextSample->GetFrames());
 		}
 
 		check(Input.Num() > 0);
@@ -391,7 +391,8 @@ uint32 FMediaAudioResampler::Generate(float* Output, FTimespan& OutTime, const u
 
 		// calculate output time
 		const float InputOffset = (FrameIndex + FrameAlpha) / InputFrames;
-		OutTime = InputTime + InputOffset * InputDuration;
+		OutTime.Time = InputTime.Time + InputOffset * InputDuration;
+		OutTime.SequenceIndex = InputTime.SequenceIndex;
 
 		// invert buffer if rate reversed
 		if ((Rate * LastRate) < 0.0f)
@@ -621,7 +622,7 @@ void FMediaAudioResampler::ClearInput()
 	InputDuration = FTimespan::Zero();
 	InputFrames = 0;
 	InputSampleRate = 0;
-	InputTime = FTimespan::Zero();
+	InputTime = FMediaTimeStamp(FTimespan::Zero());
 
 	LastRate = 0.0f;
 }
diff --git a/Engine/Source/Runtime/MediaUtils/Private/MediaHelpers.cpp b/Engine/Source/Runtime/MediaUtils/Private/MediaHelpers.cpp
--- a/Engine/Source/Runtime/MediaUtils/Private/MediaHelpers.cpp
+++ b/Engine/Source/Runtime/MediaUtils/Private/MediaHelpers.cpp
@@ -24,6 +24,8 @@ namespace MediaUtils
 		case EMediaEvent::SeekCompleted: return TEXT("SeekCompleted");
 		case EMediaEvent::TracksChanged: return TEXT("TracksChanged");
 
+		case EMediaEvent::Internal_PurgeVideoSamplesHint: return TEXT("Internal: PurgeVideoSamplesHint");
+
 		// no default case; all cases must be implemented!
 		}
 
diff --git a/Engine/Source/Runtime/MediaUtils/Private/MediaPlayerFacade.cpp b/Engine/Source/Runtime/MediaUtils/Private/MediaPlayerFacade.cpp
--- a/Engine/Source/Runtime/MediaUtils/Private/MediaPlayerFacade.cpp
+++ b/Engine/Source/Runtime/MediaUtils/Private/MediaPlayerFacade.cpp
@@ -26,6 +26,7 @@
 #include "MediaHelpers.h"
 #include "MediaSampleCache.h"
 #include "MediaSampleQueueDepths.h"
+#include "MediaSampleQueue.h"
 
 #define MEDIAPLAYERFACADE_DISABLE_BLOCKING 0
 #define MEDIAPLAYERFACADE_TRACE_SINKOVERFLOWS 0
@@ -53,8 +54,24 @@ DECLARE_CYCLE_STAT(TEXT("MediaUtils MediaPlayerFacade TickOutput"), STAT_MediaUt
 DECLARE_CYCLE_STAT(TEXT("MediaUtils MediaPlayerFacade TickTickable"), STAT_MediaUtils_FacadeTickTickable, STATGROUP_Media);
 
 /** Player time on main thread during last fetch tick. */
-DECLARE_FLOAT_COUNTER_STAT(TEXT("MediaPlayerFacade TickPostEngine Time"), STAT_MediaUtils_FacadeTime, STATGROUP_Media);
+DECLARE_FLOAT_COUNTER_STAT(TEXT("MediaPlayerFacade PlaybackTime"), STAT_MediaUtils_FacadeTime, STATGROUP_Media);
 
+/** Number of video samples currently in the queue. */
+DECLARE_DWORD_COUNTER_STAT(TEXT("MediaPlayerFacade NumVideoSamples"), STAT_MediaUtils_FacadeNumVideoSamples, STATGROUP_Media);
+
+/** Number of audio samples currently in the queue. */
+DECLARE_DWORD_COUNTER_STAT(TEXT("MediaPlayerFacade NumAudioSamples"), STAT_MediaUtils_FacadeNumAudioSamples, STATGROUP_Media);
+
+/** Number of purged video samples */
+DECLARE_DWORD_COUNTER_STAT(TEXT("MediaPlayerFacade NumPurgedVideoSamples"), STAT_MediaUtils_FacadeNumPurgedVideoSamples, STATGROUP_Media);
+DECLARE_DWORD_ACCUMULATOR_STAT(TEXT("MediaPlayerFacade TotalPurgedVideoSamples"), STAT_MediaUtils_FacadeTotalPurgedVideoSamples, STATGROUP_Media);
+
+/* Some constants
+*****************************************************************************/
+
+const double kMaxTimeSinceFrameStart = 0.300; // max seconds we allow between the start of the frame and the player facade timing computations (to catch suspended apps & debugging)
+const double kMaxTimeSinceAudioTimeSampling = 0.250; // max seconds we allow to have passed between the last audio timing sampling and the player facade timing computations (to catch suspended apps & debugging - some platforms do update audio at a farily low rate: hence the big tollerance)
+const double kOutdatedVideoSamplesTollerance = 0.050; // seconds video samples are allowed to be "too old" to stay in the player's output queue despite of calculations indicating they need to go
 
 /* Local helpers
 *****************************************************************************/
@@ -76,8 +93,10 @@ FMediaPlayerFacade::FMediaPlayerFacade()
 	, BlockOnTime(FTimespan::MinValue())
 	, Cache(new FMediaSampleCache)
 	, LastRate(0.0f)
-	, LastAudioRenderedSampleTime(FTimespan::MinValue())
-{ }
+	, bHaveActiveAudio(false)
+{ 
+	MediaModule = FModuleManager::LoadModulePtr<IMediaModule>("Media");
+}
 
 
 FMediaPlayerFacade::~FMediaPlayerFacade()
@@ -98,33 +117,34 @@ FMediaPlayerFacade::~FMediaPlayerFacade()
 /* FMediaPlayerFacade interface
 *****************************************************************************/
 
-void FMediaPlayerFacade::AddAudioSampleSink(const TSharedRef<TMediaSampleSink<IMediaAudioSample>, ESPMode::ThreadSafe>& SampleSink)
+void FMediaPlayerFacade::AddAudioSampleSink(const TSharedRef<FMediaAudioSampleSink, ESPMode::ThreadSafe>& SampleSink)
 {
 	FScopeLock Lock(&CriticalSection);
 	AudioSampleSinks.Add(SampleSink);
+	PrimaryAudioSink = AudioSampleSinks.GetPrimaryAudioSink();
 }
 
 
-void FMediaPlayerFacade::AddCaptionSampleSink(const TSharedRef<TMediaSampleSink<IMediaOverlaySample>, ESPMode::ThreadSafe>& SampleSink)
+void FMediaPlayerFacade::AddCaptionSampleSink(const TSharedRef<FMediaOverlaySampleSink, ESPMode::ThreadSafe>& SampleSink)
 {
 	CaptionSampleSinks.Add(SampleSink);
 }
 	
 
-void FMediaPlayerFacade::AddMetadataSampleSink(const TSharedRef<TMediaSampleSink<IMediaBinarySample>, ESPMode::ThreadSafe>& SampleSink)
+void FMediaPlayerFacade::AddMetadataSampleSink(const TSharedRef<FMediaBinarySampleSink, ESPMode::ThreadSafe>& SampleSink)
 {
 	FScopeLock Lock(&CriticalSection);
 	MetadataSampleSinks.Add(SampleSink);
 }
 
 
-void FMediaPlayerFacade::AddSubtitleSampleSink(const TSharedRef<TMediaSampleSink<IMediaOverlaySample>, ESPMode::ThreadSafe>& SampleSink)
+void FMediaPlayerFacade::AddSubtitleSampleSink(const TSharedRef<FMediaOverlaySampleSink, ESPMode::ThreadSafe>& SampleSink)
 {
 	SubtitleSampleSinks.Add(SampleSink);
 }
 
 
-void FMediaPlayerFacade::AddVideoSampleSink(const TSharedRef<TMediaSampleSink<IMediaTextureSample>, ESPMode::ThreadSafe>& SampleSink)
+void FMediaPlayerFacade::AddVideoSampleSink(const TSharedRef<FMediaTextureSampleSink, ESPMode::ThreadSafe>& SampleSink)
 {
 	VideoSampleSinks.Add(SampleSink);
 }
@@ -138,8 +158,6 @@ bool FMediaPlayerFacade::CanPause() const
 
 bool FMediaPlayerFacade::CanPlayUrl(const FString& Url, const IMediaOptions* Options)
 {
-	IMediaModule* MediaModule = FModuleManager::LoadModulePtr<IMediaModule>("Media");
-
 	if (MediaModule == nullptr)
 	{
 		return false;
@@ -291,6 +309,18 @@ TRangeSet<float> FMediaPlayerFacade::GetSupportedRates(bool Unthinned) const
 }
 
 
+bool FMediaPlayerFacade::HaveVideoPlayback() const
+{
+	return VideoSampleSinks.Num() && GetSelectedTrack(EMediaTrackType::Video) != INDEX_NONE;
+}
+
+
+bool FMediaPlayerFacade::HaveAudioPlayback() const
+{
+	return PrimaryAudioSink.IsValid() && GetSelectedTrack(EMediaTrackType::Audio) != INDEX_NONE;
+}
+
+
 FTimespan FMediaPlayerFacade::GetTime() const
 {
 	if (!Player.IsValid())
@@ -298,15 +328,58 @@ FTimespan FMediaPlayerFacade::GetTime() const
 		return FTimespan::Zero(); // no media opened
 	}
 
-	FTimespan Result = Player->GetControls().GetTime() - TimeDelay;
-	if (Result.GetTicks() < 0)
+	FTimespan Result;
+
+	if (Player->GetPlayerFeatureFlag(IMediaPlayer::EFeatureFlag::UsePlaybackTimingV2))
 	{
-		Result = FTimespan::Zero();
+		// New style: framework controls timing - we use GetTimeStamp() and return the legacy part of the value
+		FMediaTimeStamp TimeStamp = GetTimeStamp();
+		return TimeStamp.IsValid() ? TimeStamp.Time : FTimespan::Zero();
 	}
+	else
+	{
+		// Old style: ask the player for timing
+		Result = Player->GetControls().GetTime() - TimeDelay;
+		if (Result.GetTicks() < 0)
+		{
+			Result = FTimespan::Zero();
+		}
+	}
+
 	return Result;
 }
 
 
+FMediaTimeStamp FMediaPlayerFacade::GetTimeStamp() const
+{
+	if (!Player.IsValid() || !Player->GetPlayerFeatureFlag(IMediaPlayer::EFeatureFlag::UsePlaybackTimingV2))
+	{
+		return FMediaTimeStamp(); // no media opened / old player timing mode
+	}
+
+	if (HaveVideoPlayback())
+	{
+		/*
+			Returning the precise time of the sample returned during TickFetch()
+		*/
+		return LastVideoSampleProcessedTime.TimeStamp;
+	}
+	else if (HaveAudioPlayback())
+	{
+		/*
+			We grab the last processed audio sample timestamp when it gets passed out to the sink(s) and keep it
+			as "the value" for the frame (on the gamethread) -- an approximation, but better then having it return
+			new values each time its called in one and the same frame...
+		*/
+		return CurrentFrameAudioTimeStamp;
+	}
+
+	// we assume video and/or audio to be present in any stream we play - otherwise: no time info
+	// (at least for now)
+	return FMediaTimeStamp();
+}
+
+
 FText FMediaPlayerFacade::GetTrackDisplayName(EMediaTrackType TrackType, int32 TrackIndex) const
 {
 	return Player.IsValid() ? Player->GetTracks().GetTrackDisplayName((EMediaTrackType)TrackType, TrackIndex) : FText::GetEmpty();
@@ -449,9 +522,6 @@ bool FMediaPlayerFacade::Open(const FString& Url, const IMediaOptions* Options,
 		return false;
 	}
 
-	static const FName MediaModuleName("Media");
-	IMediaModule* MediaModule;
-	MediaModule = FModuleManager::GetModulePtr<IMediaModule>(MediaModuleName);
 	check(MediaModule);
 
 	// find & initialize new player
@@ -493,6 +563,10 @@ bool FMediaPlayerFacade::Open(const FString& Url, const IMediaOptions* Options,
 		return false;
 	}
 
+	LastVideoSampleProcessedTime.Invalidate();
+	LastAudioSampleProcessedTime.Invalidate();
+	CurrentFrameAudioTimeStamp.Invalidate();
+
 	return true;
 }
 
@@ -594,6 +668,17 @@ bool FMediaPlayerFacade::SetRate(float Rate)
 	{
 		FlushSinks(); // direction change
 	}
+	else
+	{
+		if (Rate == 0.0f)
+		{
+			// Invalidate audio time on entering pause mode...
+			if (TSharedPtr< FMediaAudioSampleSink, ESPMode::ThreadSafe> AudioSink = PrimaryAudioSink.Pin())
+			{
+				AudioSink->InvalidateAudioTime();
+			}
+		}
+	}
 
 	return true;
 }
@@ -642,22 +727,15 @@ bool FMediaPlayerFacade::SupportsRate(float Rate, bool Unthinned) const
 
 void FMediaPlayerFacade::SetLastAudioRenderedSampleTime(FTimespan SampleTime)
 {
-	LastAudioRenderedSampleTime = SampleTime;
+	FScopeLock Lock(&LastTimeValuesCS);
+	LastAudioRenderedSampleTime.TimeStamp = FMediaTimeStamp(SampleTime);
+	LastAudioRenderedSampleTime.SampledAtTime = FPlatformTime::Seconds();
 }
 
 FTimespan FMediaPlayerFacade::GetLastAudioRenderedSampleTime() const
 {
-	return LastAudioRenderedSampleTime.Load();
-}
-
-FTimespan FMediaPlayerFacade::GetLastAudioSampleProcessedTime() const
-{
-	return LastAudioSampleProcessedTime.Load();
-}
-
-FTimespan FMediaPlayerFacade::GetLastVideoSampleProcessedTime() const
-{
-	return LastVideoSampleProcessedTime.Load();
+	FScopeLock Lock(&LastTimeValuesCS);
+	return LastAudioRenderedSampleTime.TimeStamp.Time;
 }
 
 /* FMediaPlayerFacade implementation
@@ -721,9 +799,10 @@ void FMediaPlayerFacade::FlushSinks()
 
 	NextVideoSampleTime = FTimespan::MinValue();
 
-	LastAudioSampleProcessedTime = FTimespan::MinValue();
-	LastVideoSampleProcessedTime = FTimespan::MinValue();
-	LastAudioRenderedSampleTime = FTimespan::MinValue();
+	// note: we do not invalidate the LastXXXXSampleProcessingTime values here -> it is more natural for a outside caller to receive the "last good time" e.g. during a seek
+	LastAudioRenderedSampleTime.Invalidate();
+
+	NextEstVideoTimeAtFrameStart.Invalidate();
 }
 
 
@@ -766,9 +845,6 @@ TSharedPtr<IMediaPlayer, ESPMode::ThreadSafe> FMediaPlayerFacade::GetPlayerForUr
 		return Player;
 	}
 
-	// load media module
-	IMediaModule* MediaModule = FModuleManager::LoadModulePtr<IMediaModule>("Media");
-
 	if (MediaModule == nullptr)
 	{
 		UE_LOG(LogMediaUtils, Error, TEXT("Failed to load Media module"));
@@ -917,15 +993,13 @@ void FMediaPlayerFacade::ProcessEvent(EMediaEvent Event)
 			// Yes, this also means: if we still have a player, it's still the one this event originated from
 
 			// If player allows: close it down all the way right now
-			if(Player.IsValid() && Player->GetPlayerFeatureFlag(IMediaPlayer::FeatureFlag::AllowShutdownOnClose))
+			if(Player.IsValid() && Player->GetPlayerFeatureFlag(IMediaPlayer::EFeatureFlag::AllowShutdownOnClose))
 			{
 				FScopeLock Lock(&CriticalSection);
 				Player.Reset();
 			}
 
 			// Stop issuing audio thread ticks until we open the player again
-			IMediaModule* MediaModule = FModuleManager::LoadModulePtr<IMediaModule>("Media");
-			check(MediaModule);
 			MediaModule->GetTicker().RemoveTickable(AsShared());
 		}
 	}
@@ -962,6 +1036,83 @@ void FMediaPlayerFacade::SelectDefaultTracks()
 /* IMediaClockSink interface
 *****************************************************************************/
 
+void FMediaPlayerFacade::TickInput(FTimespan DeltaTime, FTimespan Timecode)
+{
+	SCOPE_CYCLE_COUNTER(STAT_MediaUtils_FacadeTickInput);
+
+	if (Player.IsValid())
+	{
+		Player->TickInput(DeltaTime, Timecode);
+
+		// Update flag reflecting presence of audio ion the current stream
+		// (doing it just once per gameloop is enough
+		bHaveActiveAudio = HaveAudioPlayback();
+
+		// get current play rate
+		float Rate = Player->GetControls().GetRate();
+
+		if (Rate == 0.0f)
+		{
+			Rate = LastRate;
+		}
+		else
+		{
+			LastRate = Rate;
+		}
+
+		if (Player->GetPlayerFeatureFlag(IMediaPlayer::EFeatureFlag::UsePlaybackTimingV2))
+		{
+			//
+			// New timing control (handled before any engine world, object etc. updates; so "all frame" (almost) see the state produced here)
+			//
+
+			// Do we have a current timestamp estimation?
+			if (!bHaveActiveAudio && !NextEstVideoTimeAtFrameStart.IsValid())
+			{
+				// Not, yet. We need to attempt to get the next video sample's timestamp to get going...
+				FMediaTimeStamp VideoTimeStamp;
+				if (Player->GetSamples().PeekVideoSampleTime(VideoTimeStamp))
+				{
+					NextEstVideoTimeAtFrameStart = VideoTimeStamp;
+				}
+			}
+
+			TRange<FMediaTimeStamp> TimeRange;
+			if (!GetCurrentPlaybackTimeRange(TimeRange, Rate, DeltaTime, false))
+			{
+				return;
+			}
+
+			SET_FLOAT_STAT(STAT_MediaUtils_FacadeTime, TimeRange.GetLowerBoundValue().Time.GetTotalSeconds());
+
+			//
+			// Process samples in range
+			//
+
+			IMediaSamples& Samples = Player->GetSamples();
+			ProcessCaptionSamples(Samples, TimeRange);
+			ProcessSubtitleSamples(Samples, TimeRange);
+			ProcessVideoSamples(Samples, TimeRange);
+
+			SET_DWORD_STAT(STAT_MediaUtils_FacadeNumVideoSamples, Samples.NumVideoSamples());
+
+			// Move video frame start estimate forward if we have no audio timing to guide us
+			if (!bHaveActiveAudio && NextEstVideoTimeAtFrameStart.IsValid())
+			{
+				NextEstVideoTimeAtFrameStart.Time += DeltaTime * Rate;
+			}
+
+			if (bHaveActiveAudio)
+			{
+				// Keep currently last processed audio sample timestamp available for all frame (to provide consistent info)
+				FScopeLock Lock(&LastTimeValuesCS);
+				CurrentFrameAudioTimeStamp = LastAudioSampleProcessedTime.TimeStamp;
+			}
+		}
+	}
+}
+
+
 void FMediaPlayerFacade::TickFetch(FTimespan DeltaTime, FTimespan Timecode)
 {
 	SCOPE_CYCLE_COUNTER(STAT_MediaUtils_FacadeTickFetch);
@@ -972,12 +1123,13 @@ void FMediaPlayerFacade::TickFetch(FTimespan DeltaTime, FTimespan Timecode)
 		Player->TickFetch(DeltaTime, Timecode);
 	}
 
-	// process deferred events
-	EMediaEvent Event;
-
-	while (QueuedEvents.Dequeue(Event))
 	{
-		ProcessEvent(Event);
+		// process deferred events
+		EMediaEvent Event;
+		while (QueuedEvents.Dequeue(Event))
+		{
+			ProcessEvent(Event);
+		}
 	}
 
 	if (!Player.IsValid())
@@ -992,81 +1144,71 @@ void FMediaPlayerFacade::TickFetch(FTimespan DeltaTime, FTimespan Timecode)
 	{
 		Rate = LastRate;
 	}
-	else
-	{
-		LastRate = Rate;
-	}
-
-	// determine range of valid samples
-	TRange<FTimespan> TimeRange;
-
-	const FTimespan CurrentTime = GetTime();
 
-	SET_FLOAT_STAT(STAT_MediaUtils_FacadeTime, CurrentTime.GetTotalMilliseconds());
-
-	if (Rate > 0.0f)
-	{
-		TimeRange = TRange<FTimespan>::AtMost(CurrentTime);
-	}
-	else if (Rate < 0.0f)
-	{
-		TimeRange = TRange<FTimespan>::AtLeast(CurrentTime);
-	}
-	else
+	if (!Player->GetPlayerFeatureFlag(IMediaPlayer::EFeatureFlag::UsePlaybackTimingV2))
 	{
-		TimeRange = TRange<FTimespan>(CurrentTime);
-	}
+		//
+		// Old timing control
+		//
+		TRange<FTimespan> TimeRange;
 
-	// process samples in range
-	IMediaSamples& Samples = Player->GetSamples();
-
-	bool Blocked = false;
-	FDateTime BlockedTime;
+		const FTimespan CurrentTime = GetTime();
 
-	while (true)
-	{
-		ProcessCaptionSamples(Samples, TimeRange);
-		ProcessSubtitleSamples(Samples, TimeRange);
-		ProcessVideoSamples(Samples, TimeRange);
+		SET_FLOAT_STAT(STAT_MediaUtils_FacadeTime, CurrentTime.GetTotalSeconds());
 
-		if (!BlockOnFetch())
+		if (Rate > 0.0f)
 		{
-			break;
+			TimeRange = TRange<FTimespan>::AtMost(CurrentTime);
 		}
-
-		if (Blocked)
+		else if (Rate < 0.0f)
 		{
-			if ((FDateTime::UtcNow() - BlockedTime) >= FTimespan::FromSeconds(MEDIAUTILS_MAX_BLOCKONFETCH_SECONDS))
-			{
-				UE_LOG(LogMediaUtils, Verbose, TEXT("PlayerFacade %p: Aborted block on fetch %s after %i seconds"),
-					this,
-					*BlockOnTime.ToString(TEXT("%h:%m:%s.%t")),
-					MEDIAUTILS_MAX_BLOCKONFETCH_SECONDS
-				);
-
-				break;
-			}
+			TimeRange = TRange<FTimespan>::AtLeast(CurrentTime);
 		}
 		else
 		{
-			UE_LOG(LogMediaUtils, VeryVerbose, TEXT("PlayerFacade %p: Blocking on fetch %s"), this, *BlockOnTime.ToString(TEXT("%h:%m:%s.%t")));
-
-			Blocked = true;
-			BlockedTime = FDateTime::UtcNow();
+			TimeRange = TRange<FTimespan>(CurrentTime);
 		}
 
-		FPlatformProcess::Sleep(0.0f);
-	}
-}
+		// process samples in range
+		IMediaSamples& Samples = Player->GetSamples();
 
+		bool Blocked = false;
+		FDateTime BlockedTime;
 
-void FMediaPlayerFacade::TickInput(FTimespan DeltaTime, FTimespan Timecode)
-{
-	SCOPE_CYCLE_COUNTER(STAT_MediaUtils_FacadeTickInput);
+		while (true)
+		{
+			ProcessCaptionSamples(Samples, TimeRange);
+			ProcessSubtitleSamples(Samples, TimeRange);
+			ProcessVideoSamples(Samples, TimeRange);
 
-	if (Player.IsValid())
-	{
-		Player->TickInput(DeltaTime, Timecode);
+			if (!BlockOnFetch())
+			{
+				break;
+			}
+
+			if (Blocked)
+			{
+				if ((FDateTime::UtcNow() - BlockedTime) >= FTimespan::FromSeconds(MEDIAUTILS_MAX_BLOCKONFETCH_SECONDS))
+				{
+					UE_LOG(LogMediaUtils, Verbose, TEXT("PlayerFacade %p: Aborted block on fetch %s after %i seconds"),
+						this,
+						*BlockOnTime.ToString(TEXT("%h:%m:%s.%t")),
+						MEDIAUTILS_MAX_BLOCKONFETCH_SECONDS
+					);
+
+					break;
+				}
+			}
+			else
+			{
+				UE_LOG(LogMediaUtils, VeryVerbose, TEXT("PlayerFacade %p: Blocking on fetch %s"), this, *BlockOnTime.ToString(TEXT("%h:%m:%s.%t")));
+
+				Blocked = true;
+				BlockedTime = FDateTime::UtcNow();
+			}
+
+			FPlatformProcess::Sleep(0.0f);
+		}
 	}
 }
 
@@ -1104,7 +1246,10 @@ void FMediaPlayerFacade::TickTickable()
 		return;
 	}
 
-	Player->SetLastAudioRenderedSampleTime(LastAudioRenderedSampleTime.Load());
+	{
+	FScopeLock Lock1(&LastTimeValuesCS);
+	Player->SetLastAudioRenderedSampleTime(LastAudioRenderedSampleTime.TimeStamp.Time);
+	}
 
 	Player->TickAudio();
 
@@ -1114,14 +1259,22 @@ void FMediaPlayerFacade::TickTickable()
 
 	const FTimespan Time = GetTime();
 
+	bool bUseV2Timing = Player->GetPlayerFeatureFlag(IMediaPlayer::EFeatureFlag::UsePlaybackTimingV2);
+
 	if (LastRate > 0.0f)
 	{
-		AudioTimeRange = TRange<FTimespan>::Inclusive(FTimespan::MinValue(), Time + MediaPlayerFacade::AudioPreroll);
+		if (!bUseV2Timing) // we leave range open - sends all the player has
+		{
+			AudioTimeRange = TRange<FTimespan>::Inclusive(FTimespan::MinValue(), Time + MediaPlayerFacade::AudioPreroll);
+		}
 		MetadataTimeRange = TRange<FTimespan>::Inclusive(FTimespan::MinValue(), Time + MediaPlayerFacade::MetadataPreroll);
 	}
 	else
 	{
-		AudioTimeRange = TRange<FTimespan>::Inclusive(Time - MediaPlayerFacade::AudioPreroll, FTimespan::MaxValue());
+		if (!bUseV2Timing) // we leave range open - sends all the player has
+		{
+			AudioTimeRange = TRange<FTimespan>::Inclusive(Time - MediaPlayerFacade::AudioPreroll, FTimespan::MaxValue());
+		}
 		MetadataTimeRange = TRange<FTimespan>::Inclusive(Time - MediaPlayerFacade::MetadataPreroll, FTimespan::MaxValue());
 	}
 
@@ -1130,8 +1283,91 @@ void FMediaPlayerFacade::TickTickable()
 	
 	ProcessAudioSamples(Samples, AudioTimeRange);
 	ProcessMetadataSamples(Samples, MetadataTimeRange);	
+
+	SET_DWORD_STAT(STAT_MediaUtils_FacadeNumAudioSamples, Samples.NumAudio());
 }
 
+
+bool FMediaPlayerFacade::GetCurrentPlaybackTimeRange(TRange<FMediaTimeStamp> & TimeRange, float Rate, FTimespan DeltaTime, bool bDoNotUseFrameStartReference) const
+{
+	check(Player->GetPlayerFeatureFlag(IMediaPlayer::EFeatureFlag::UsePlaybackTimingV2));
+
+	TSharedPtr<FMediaAudioSampleSink, ESPMode::ThreadSafe> AudioSink = PrimaryAudioSink.Pin();
+
+	if (bHaveActiveAudio && AudioSink.IsValid())
+	{
+		//
+		// Audio is available...
+		//
+
+		FMediaTimeStampSample AudioTime = AudioSink->GetAudioTime();
+		if (!AudioTime.IsValid())
+		{
+			// No timing info available, no timerange available, no samples to process
+			return false;
+		}
+
+		FMediaTimeStamp EstAudioTimeAtFrameStart;
+
+		double Now = FPlatformTime::Seconds();
+
+		if (!bDoNotUseFrameStartReference)
+		{
+			// Normal estimation relative to current frame start...
+			// (on gamethread operation)
+
+			check(IsInGameThread());
+
+			double AgeOfFrameStart = Now - MediaModule->GetFrameStartTime();
+			double AgeOfAudioTime = Now - AudioTime.SampledAtTime;
+
+			if (AgeOfFrameStart >= 0.0 && AgeOfFrameStart <= kMaxTimeSinceFrameStart &&
+				AgeOfAudioTime >= 0.0 && AgeOfAudioTime <= kMaxTimeSinceAudioTimeSampling)
+			{
+				// All realtime timestamps seem in sane ranges - we most likely did not have a lengthy interruption (suspended / debugging step)
+				EstAudioTimeAtFrameStart = AudioTime.TimeStamp + FTimespan::FromSeconds((MediaModule->GetFrameStartTime() - AudioTime.SampledAtTime) * Rate);
+			}
+			else
+			{
+				// Realtime timestamps seem wonky. Proceed without them (worse estimation quality)
+				EstAudioTimeAtFrameStart = AudioTime.TimeStamp;
+			}
+		}
+		else
+		{
+			// Do not use frame start reference -> we compute relative to "now"
+			// (for use off gamethread)
+			EstAudioTimeAtFrameStart = AudioTime.TimeStamp + FTimespan::FromSeconds((Now - AudioTime.SampledAtTime) * Rate);
+		}
+
+		TimeRange = (Rate >= 0.0f) ? TRange<FMediaTimeStamp>(EstAudioTimeAtFrameStart, EstAudioTimeAtFrameStart + DeltaTime * Rate)
+								   : TRange<FMediaTimeStamp>(EstAudioTimeAtFrameStart + DeltaTime * Rate, EstAudioTimeAtFrameStart);
+	}
+	else
+	{
+		//
+		// No Audio (no data and/or no sink)
+		//
+
+		// Do we now have a current timestamp estimation?
+		if (!NextEstVideoTimeAtFrameStart.IsValid())
+		{
+			// No timing info available, no timerange available, no samples to process
+			return false;
+		}
+		else
+		{
+			// Yes. Setup current time range & advance time estimation...
+
+			TimeRange = (Rate >= 0.0f) ? TRange<FMediaTimeStamp>(NextEstVideoTimeAtFrameStart, NextEstVideoTimeAtFrameStart + DeltaTime * Rate)
+									   : TRange<FMediaTimeStamp>(NextEstVideoTimeAtFrameStart + DeltaTime * Rate, NextEstVideoTimeAtFrameStart);
+		}
+	}
+
+	return true;
+}
+
+
 /* FMediaPlayerFacade implementation
 *****************************************************************************/
 
@@ -1139,20 +1375,55 @@ void FMediaPlayerFacade::ProcessAudioSamples(IMediaSamples& Samples, TRange<FTim
 {
 	TSharedPtr<IMediaAudioSample, ESPMode::ThreadSafe> Sample;
 
-	while (Samples.FetchAudio(TimeRange, Sample))
+	if (AudioSampleSinks.Num() == 1)
 	{
-		if (!Sample.IsValid())
+		//
+		// "Modern" 1-Audio-Sink-Only case
+		//
+		if (TSharedPtr< FMediaAudioSampleSink, ESPMode::ThreadSafe> PinnedPrimaryAudioSink = PrimaryAudioSink.Pin())
 		{
-			continue;
-		}
+			while (PinnedPrimaryAudioSink->CanAcceptSamples(1))
+			{
+				if (!Samples.FetchAudio(TimeRange, Sample))
+					break;
 
-		LastAudioSampleProcessedTime = Sample->GetTime();
+				if (!Sample.IsValid())
+				{
+					continue;
+				}
 
-		if (!AudioSampleSinks.Enqueue(Sample.ToSharedRef(), FMediaPlayerQueueDepths::MaxAudioSinkDepth))
+				LastAudioSampleProcessedTime.TimeStamp = FMediaTimeStamp(Sample->GetTime());
+				LastAudioSampleProcessedTime.SampledAtTime = FPlatformTime::Seconds();
+
+				// We assume we are the only ones adding samples: hence - after the above check - this MUST succeed!
+				verify(AudioSampleSinks.Enqueue(Sample.ToSharedRef(), FMediaPlayerQueueDepths::MaxAudioSinkDepth));
+			}
+		}
+	}
+	else
+	{
+		//
+		// >1 Audio Sinks: we must drop samples that cause on overrun as SOME sinks will get it, some don't...
+		// (mainly here to cover the "backwards compatibility" cases -> in the future we will probably only allow ONE AudioSink)
+		//
+		while (Samples.FetchAudio(TimeRange, Sample))
 		{
-			#if MEDIAPLAYERFACADE_TRACE_SINKOVERFLOWS
+			if (!Sample.IsValid())
+			{
+				continue;
+			}
+
+			if (!AudioSampleSinks.Enqueue(Sample.ToSharedRef(), FMediaPlayerQueueDepths::MaxAudioSinkDepth))
+			{
+#if MEDIAPLAYERFACADE_TRACE_SINKOVERFLOWS
 				UE_LOG(LogMediaUtils, VeryVerbose, TEXT("PlayerFacade %p: Audio sample sink overflow"), this);
-			#endif
+#endif
+			}
+			else
+			{
+				LastAudioSampleProcessedTime.TimeStamp = FMediaTimeStamp(Sample->GetTime());
+				LastAudioSampleProcessedTime.SampledAtTime = FPlatformTime::Seconds();
+			}
 		}
 	}
 }
@@ -1214,6 +1485,9 @@ void FMediaPlayerFacade::ProcessVideoSamples(IMediaSamples& Samples, TRange<FTim
 		Player->ProcessVideoSamples();
 	}
 
+	// This is not to be used with V2 timing
+	check(!Player->GetPlayerFeatureFlag(IMediaPlayer::EFeatureFlag::UsePlaybackTimingV2));
+
 	TSharedPtr<IMediaTextureSample, ESPMode::ThreadSafe> Sample;
 
 	while (Samples.FetchVideo(TimeRange, Sample))
@@ -1223,23 +1497,131 @@ void FMediaPlayerFacade::ProcessVideoSamples(IMediaSamples& Samples, TRange<FTim
 			continue;
 		}
 
-		LastVideoSampleProcessedTime = Sample->GetTime();
+		{
+		FScopeLock Lock(&LastTimeValuesCS);
+		LastVideoSampleProcessedTime.TimeStamp = FMediaTimeStamp(Sample->GetTime());
+		LastVideoSampleProcessedTime.SampledAtTime = FPlatformTime::Seconds();
+		}
 
-		UE_LOG(LogMediaUtils, VeryVerbose, TEXT("PlayerFacade %p: Fetched video sample %s"), this, *Sample->GetTime().ToString(TEXT("%h:%m:%s.%t")));
+		UE_LOG(LogMediaUtils, VeryVerbose, TEXT("PlayerFacade %p: Fetched video sample %s"), this, *Sample->GetTime().Time.ToString(TEXT("%h:%m:%s.%t")));
 
 		if (VideoSampleSinks.Enqueue(Sample.ToSharedRef(), FMediaPlayerQueueDepths::MaxVideoSinkDepth))
 		{
 			if (GetRate() >= 0.0f)
 			{
-				NextVideoSampleTime = Sample->GetTime() + Sample->GetDuration();
+				NextVideoSampleTime = Sample->GetTime().Time + Sample->GetDuration();
 				UE_LOG(LogMediaUtils, VeryVerbose, TEXT("PlayerFacade %p: Next video sample time %s"), this, *NextVideoSampleTime.ToString(TEXT("%h:%m:%s.%t")));
 			}
 		}
 		else
 		{
-			#if MEDIAPLAYERFACADE_TRACE_SINKOVERFLOWS
-				UE_LOG(LogMediaUtils, VeryVerbose, TEXT("PlayerFacade %p: Video sample sink overflow"), this);
-			#endif
+#if MEDIAPLAYERFACADE_TRACE_SINKOVERFLOWS
+			UE_LOG(LogMediaUtils, VeryVerbose, TEXT("PlayerFacade %p: Video sample sink overflow"), this);
+#endif
+		}
+	}
+}
+
+
+void FMediaPlayerFacade::ProcessVideoSamples(IMediaSamples& Samples, const TRange<FMediaTimeStamp> & TimeRange)
+{
+	// Let the player do some processing if needed.
+	if (Player.IsValid())
+	{
+		// note: avoid using this - it will be deprecated
+		Player->ProcessVideoSamples();
+	}
+
+	// This is not to be used with V1 timing
+	check(Player->GetPlayerFeatureFlag(IMediaPlayer::EFeatureFlag::UsePlaybackTimingV2));
+	// We expect a fully closed range or we assume: nothing to do...
+	check(TimeRange.GetLowerBound().IsClosed() && TimeRange.GetUpperBound().IsClosed());
+
+	TSharedPtr<IMediaTextureSample, ESPMode::ThreadSafe> Sample;
+
+	switch (Samples.FetchBestVideoSampleForTimeRange(TimeRange, Sample, Player->GetControls().GetRate() < 0.0f))
+	{
+		case IMediaSamples::EFetchBestSampleResult::Ok:
+		{
+			// Enqueue the sample to render
+			// (we use a queue to stay compatible with existing structure and older sinks - new sinks will read this single entry right away on the gamethread
+			//  and pass it along to rendering outside the queue)
+			bool bOk = VideoSampleSinks.Enqueue(Sample.ToSharedRef(), FMediaPlayerQueueDepths::MaxVideoSinkDepth);
+			check(bOk);
+
+			FScopeLock Lock(&LastTimeValuesCS);
+			LastVideoSampleProcessedTime.TimeStamp = FMediaTimeStamp(Sample->GetTime());
+			LastVideoSampleProcessedTime.SampledAtTime = FPlatformTime::Seconds();
+
+			//uint32 Frames = uint32(LastVideoSampleProcessedTime.TimeStamp.Time.GetTotalMilliseconds() / Sample->GetDuration().GetTotalMilliseconds() + 0.5);
+			//UE_LOG(LogMediaUtils, Display, TEXT("TIME: F=%d"), Frames % 30);
+			break;
+		}
+
+		case IMediaSamples::EFetchBestSampleResult::NoSample:
+			break;
+
+		case IMediaSamples::EFetchBestSampleResult::NotSupported:
+		{
+			//
+			// Fallback for players supporting V2 timing, but do not supply FetchBestVideoSampleForTimeRange() due to some
+			// custom implementation of IMediaSamples (here to ease adoption of the new timing code - eventually should go away)
+			//
+			const bool bReverse = (Player->GetControls().GetRate() < 0.0f);
+
+			// Find newest sample that satisfies the time range
+			// (the FetchXYZ() code does not work well with a lower range limit at all - we ask for a "up to" type range instead
+			//  and limit the other side of the range in code here to not change the older logic & possibly cause trouble in old code)
+			TRange<FMediaTimeStamp> TempRange = bReverse ? TRange<FMediaTimeStamp>::AtLeast(TimeRange.GetUpperBoundValue()) : TRange<FMediaTimeStamp>::AtMost(TimeRange.GetUpperBoundValue());
+			while (Samples.FetchVideo(TempRange, Sample))
+				;
+			if (Sample.IsValid() &&
+				((!bReverse && ((Sample->GetTime() + Sample->GetDuration()) > TimeRange.GetLowerBoundValue())) ||
+				(bReverse && ((Sample->GetTime() - Sample->GetDuration()) < TimeRange.GetLowerBoundValue()))))
+			{
+				// Enqueue the sample to render
+				// (we use a queue to stay compatible with existing structure and older sinks - new sinks will read this single entry right away on the gamethread
+				//  and pass it along to rendering outside the queue)
+				bool bOk = VideoSampleSinks.Enqueue(Sample.ToSharedRef(), FMediaPlayerQueueDepths::MaxVideoSinkDepth);
+				check(bOk);
+
+				FScopeLock Lock(&LastTimeValuesCS);
+				LastVideoSampleProcessedTime.TimeStamp = Sample->GetTime();
+				LastVideoSampleProcessedTime.SampledAtTime = FPlatformTime::Seconds();
+			}
+			break;
+		}
+	}
+}
+
+
+void FMediaPlayerFacade::ProcessCaptionSamples(IMediaSamples& Samples, TRange<FMediaTimeStamp> TimeRange)
+{
+	TSharedPtr<IMediaOverlaySample, ESPMode::ThreadSafe> Sample;
+
+	while (Samples.FetchCaption(TimeRange, Sample))
+	{
+		if (Sample.IsValid() && !CaptionSampleSinks.Enqueue(Sample.ToSharedRef(), FMediaPlayerQueueDepths::MaxCaptionSinkDepth))
+		{
+#if MEDIAPLAYERFACADE_TRACE_SINKOVERFLOWS
+			UE_LOG(LogMediaUtils, VeryVerbose, TEXT("PlayerFacade %p: Caption sample sink overflow"), this);
+#endif
+		}
+	}
+}
+
+
+void FMediaPlayerFacade::ProcessSubtitleSamples(IMediaSamples& Samples, TRange<FMediaTimeStamp> TimeRange)
+{
+	TSharedPtr<IMediaOverlaySample, ESPMode::ThreadSafe> Sample;
+
+	while (Samples.FetchSubtitle(TimeRange, Sample))
+	{
+		if (Sample.IsValid() && !SubtitleSampleSinks.Enqueue(Sample.ToSharedRef(), FMediaPlayerQueueDepths::MaxSubtitleSinkDepth))
+		{
+#if MEDIAPLAYERFACADE_TRACE_SINKOVERFLOWS
+			UE_LOG(LogMediaUtils, VeryVerbose, TEXT("PlayerFacade %p: Subtitle sample sink overflow"), this);
+#endif
 		}
 	}
 }
@@ -1251,5 +1633,52 @@ void FMediaPlayerFacade::ProcessVideoSamples(IMediaSamples& Samples, TRange<FTim
 void FMediaPlayerFacade::ReceiveMediaEvent(EMediaEvent Event)
 {
 	UE_LOG(LogMediaUtils, VeryVerbose, TEXT("PlayerFacade %p: Received media event %s"), this, *MediaUtils::EventToString(Event));
-	QueuedEvents.Enqueue(Event);
+
+	if (Event >= EMediaEvent::Internal_Start)
+	{
+		switch(Event)
+		{
+			case	EMediaEvent::Internal_PurgeVideoSamplesHint:
+			{
+				//
+				// Player asks to attempt to purge older samples in the video output queue it maintains
+				// (ask goes via facade as the player does not have accurate timing info)
+				//
+
+				// We only support this for V2 timing players
+				check(Player->GetPlayerFeatureFlag(IMediaPlayer::EFeatureFlag::UsePlaybackTimingV2));
+
+				if (!Player.IsValid())
+					return;
+
+				float Rate = Player->GetControls().GetRate();
+				if (Rate == 0.0f)
+				{
+					return;
+				}
+
+				// Get current playback time
+				// (note: we have DeltaTime forced to zero -> we just get a single value & we compute relative to "now", noty any game frame start)
+				TRange<FMediaTimeStamp> TimeRange;
+				if (!GetCurrentPlaybackTimeRange(TimeRange, Rate, FTimespan::Zero(), true))
+				{
+					return;
+				}
+
+				bool bReverse = (Rate < 0.0f);
+				uint32 NumPurged = Player->GetSamples().PurgeOutdatedVideoSamples(TimeRange.GetLowerBoundValue() + (bReverse ? kOutdatedVideoSamplesTollerance : -kOutdatedVideoSamplesTollerance), bReverse);
+				
+				SET_DWORD_STAT(STAT_MediaUtils_FacadeNumPurgedVideoSamples, NumPurged);
+				INC_DWORD_STAT_BY(STAT_MediaUtils_FacadeTotalPurgedVideoSamples, NumPurged);
+
+				break;
+			}
+			default:
+				break;
+		}
+	}
+	else
+	{
+		QueuedEvents.Enqueue(Event);
+	}
 }
diff --git a/Engine/Source/Runtime/MediaUtils/Private/MediaRecorder.cpp b/Engine/Source/Runtime/MediaUtils/Private/MediaRecorder.cpp
--- a/Engine/Source/Runtime/MediaUtils/Private/MediaRecorder.cpp
+++ b/Engine/Source/Runtime/MediaUtils/Private/MediaRecorder.cpp
@@ -335,7 +335,7 @@ void FMediaRecorder::TickRecording()
 			}
 			else
 			{
-				ImageTask->Filename = FString::Printf(TEXT("%s_%.16lu"), *BaseFilename, Sample->GetTime().GetTicks());
+				ImageTask->Filename = FString::Printf(TEXT("%s_%.16lu"), *BaseFilename, Sample->GetTime().Time.GetTicks());
 			}
 
 			ImageWriteQueue->Enqueue(MoveTemp(ImageTask), false);
diff --git a/Engine/Source/Runtime/MediaUtils/Private/MediaSampleCache.cpp b/Engine/Source/Runtime/MediaUtils/Private/MediaSampleCache.cpp
--- a/Engine/Source/Runtime/MediaUtils/Private/MediaSampleCache.cpp
+++ b/Engine/Source/Runtime/MediaUtils/Private/MediaSampleCache.cpp
@@ -74,7 +74,7 @@ void FMediaSampleCache::GetCachedAudioSampleRanges(TRangeSet<FTimespan>& OutTime
 
 	for (const TSharedPtr<IMediaAudioSample, ESPMode::ThreadSafe>& Sample : AudioSamples)
 	{
-		const FTimespan SampleTime = Sample->GetTime();
+		const FTimespan SampleTime = Sample->GetTime().Time;
 		OutTimeRanges.Add(TRange<FTimespan>(SampleTime, SampleTime + Sample->GetDuration()));
 	}
 }
@@ -86,7 +86,7 @@ void FMediaSampleCache::GetCachedVideoSampleRanges(TRangeSet<FTimespan>& OutTime
 
 	for (const TSharedPtr<IMediaTextureSample, ESPMode::ThreadSafe>& Sample : VideoSamples)
 	{
-		const FTimespan SampleTime = Sample->GetTime();
+		const FTimespan SampleTime = Sample->GetTime().Time;
 		OutTimeRanges.Add(TRange<FTimespan>(SampleTime, SampleTime + Sample->GetDuration()));
 	}
 }
diff --git a/Engine/Source/Runtime/MediaUtils/Private/MediaSamples.cpp b/Engine/Source/Runtime/MediaUtils/Private/MediaSamples.cpp
--- a/Engine/Source/Runtime/MediaUtils/Private/MediaSamples.cpp
+++ b/Engine/Source/Runtime/MediaUtils/Private/MediaSamples.cpp
@@ -7,31 +7,58 @@
 #include "IMediaOverlaySample.h"
 #include "IMediaTextureSample.h"
 
+const uint32 MaxNumberOfQueuedVideoSamples = 4;
 
 /* Local helpers
 *****************************************************************************/
 
-template<typename SampleType>
-bool FetchSample(TMediaSampleQueue<SampleType>& SampleQueue, TRange<FTimespan> TimeRange, TSharedPtr<SampleType, ESPMode::ThreadSafe>& OutSample)
+template<typename SampleType, typename SinkType>
+bool FetchSample(TMediaSampleQueue<SampleType, SinkType>& SampleQueue, TRange<FTimespan> TimeRange, TSharedPtr<SampleType, ESPMode::ThreadSafe>& OutSample)
 {
-	if (!SampleQueue.Peek(OutSample))
+	TSharedPtr<SampleType, ESPMode::ThreadSafe> Sample;
+
+	if (!SampleQueue.Peek(Sample))
 	{
 		return false;
 	}
 
-	const FTimespan SampleTime = OutSample->GetTime();
+	const FTimespan SampleTime = Sample->GetTime().Time;
 
-	if (!TimeRange.Overlaps(TRange<FTimespan>(SampleTime, SampleTime + OutSample->GetDuration())))
+	if (!TimeRange.Overlaps(TRange<FTimespan>(SampleTime, SampleTime + Sample->GetDuration())))
 	{
 		return false;
 	}
 
+	OutSample = Sample;
 	SampleQueue.Pop();
 
 	return true;
 }
 
 
+template<typename SampleType, typename SinkType>
+bool FetchSample(TMediaSampleQueue<SampleType, SinkType>& SampleQueue, TRange<FMediaTimeStamp> TimeRange, TSharedPtr<SampleType, ESPMode::ThreadSafe>& OutSample)
+{
+	TSharedPtr<SampleType, ESPMode::ThreadSafe> Sample;
+
+	if (!SampleQueue.Peek(Sample))
+	{
+		return false;
+	}
+
+	const FMediaTimeStamp SampleTime = Sample->GetTime();
+
+	if (!TimeRange.Overlaps(TRange<FMediaTimeStamp>(SampleTime, SampleTime + Sample->GetDuration())))
+	{
+		return false;
+	}
+
+	OutSample = Sample;
+	SampleQueue.Pop();
+
+	return true;
+}
+
 /* IMediaSamples interface
 *****************************************************************************/
 
@@ -64,12 +91,57 @@ bool FMediaSamples::FetchVideo(TRange<FTimespan> TimeRange, TSharedPtr<IMediaTex
 	return FetchSample(VideoSampleQueue, TimeRange, OutSample);
 }
 
+bool FMediaSamples::FetchCaption(TRange<FMediaTimeStamp> TimeRange, TSharedPtr<IMediaOverlaySample, ESPMode::ThreadSafe>& OutSample)
+{
+	return FetchSample(CaptionSampleQueue, TimeRange, OutSample);
+}
+
+bool FMediaSamples::FetchSubtitle(TRange<FMediaTimeStamp> TimeRange, TSharedPtr<IMediaOverlaySample, ESPMode::ThreadSafe>& OutSample)
+{
+	return FetchSample(SubtitleSampleQueue, TimeRange, OutSample);
+}
+
+bool FMediaSamples::FetchVideo(TRange<FMediaTimeStamp> TimeRange, TSharedPtr<IMediaTextureSample, ESPMode::ThreadSafe>& OutSample)
+{
+	return FetchSample(VideoSampleQueue, TimeRange, OutSample);
+}
 
 void FMediaSamples::FlushSamples()
 {
+	// Flushing may have various side effects that better all happen on the gamethread
+	check(IsInGameThread());
+
 	AudioSampleQueue.RequestFlush();
 	MetadataSampleQueue.RequestFlush();
 	CaptionSampleQueue.RequestFlush();
 	SubtitleSampleQueue.RequestFlush();
 	VideoSampleQueue.RequestFlush();
 }
+
+/**
+ * Fetch video sample best suited for the given time range. Samples prior to the selected one will be removed from the queue.
+ */
+FMediaSamples::EFetchBestSampleResult FMediaSamples::FetchBestVideoSampleForTimeRange(const TRange<FMediaTimeStamp> & TimeRange, TSharedPtr<IMediaTextureSample, ESPMode::ThreadSafe>& OutSample, bool bReverse)
+{
+	if (!VideoSampleQueue.FetchBestSampleForTimeRange(TimeRange, OutSample, bReverse))
+	{
+		return EFetchBestSampleResult::NoSample;
+	}
+	return EFetchBestSampleResult::Ok;
+}
+
+/**
+ * Remove any video samples from the queue that have no chance of being displayed anymore
+ */
+uint32 FMediaSamples::PurgeOutdatedVideoSamples(const FMediaTimeStamp & ReferenceTime, bool bReversed)
+{
+	return VideoSampleQueue.PurgeOutdatedSamples(ReferenceTime, bReversed);
+}
+
+/**
+ * Check if can receive more video samples
+ */
+bool FMediaSamples::CanReceiveVideoSamples(uint32 Num) const
+{
+	return (VideoSampleQueue.Num() + Num) <= MaxNumberOfQueuedVideoSamples;
+}
diff --git a/Engine/Source/Runtime/MediaUtils/Public/MediaAudioResampler.h b/Engine/Source/Runtime/MediaUtils/Public/MediaAudioResampler.h
--- a/Engine/Source/Runtime/MediaUtils/Public/MediaAudioResampler.h
+++ b/Engine/Source/Runtime/MediaUtils/Public/MediaAudioResampler.h
@@ -7,6 +7,7 @@
 #include "Delegates/Delegate.h"
 #include "MediaSampleSource.h"
 #include "Templates/SharedPointer.h"
+#include "IMediaTimeSource.h"
 
 class IMediaAudioSample;
 
@@ -38,7 +39,7 @@ public:
 	 * @param JumpFrame Frame that a jump in time occured on.
 	 * @return The actual number of frames returned.
 	 */
-	uint32 Generate(float* Output, FTimespan& OutTime, const uint32 FramesRequested, float Rate, FTimespan Time, FMediaAudioSampleSource& SampleSource, uint32& JumpFrame);
+	uint32 Generate(float* Output, FMediaTimeStamp& OutTime, const uint32 FramesRequested, float Rate, FTimespan Time, FMediaAudioSampleSource& SampleSource, uint32& JumpFrame);
 
 	/**
 	 * Initialize the resampler.
@@ -105,7 +106,7 @@ private:
 	uint32 InputSampleRate;
 
 	/** Start time of the input buffer. */
-	FTimespan InputTime;
+	FMediaTimeStamp InputTime;
 
 	/** Number of channels in the output. */
 	uint32 OutputChannels;
diff --git a/Engine/Source/Runtime/MediaUtils/Public/MediaPlayerFacade.h b/Engine/Source/Runtime/MediaUtils/Public/MediaPlayerFacade.h
--- a/Engine/Source/Runtime/MediaUtils/Public/MediaPlayerFacade.h
+++ b/Engine/Source/Runtime/MediaUtils/Public/MediaPlayerFacade.h
@@ -9,6 +9,7 @@
 #include "IMediaClockSink.h"
 #include "IMediaEventSink.h"
 #include "IMediaTickable.h"
+#include "IMediaTimeSource.h"
 #include "MediaPlayerOptions.h"
 #include "Math/Quat.h"
 #include "Math/Range.h"
@@ -24,6 +25,7 @@ class FMediaSampleCache;
 class IMediaOptions;
 class IMediaPlayer;
 class IMediaSamples;
+class IMediaModule;
 
 enum class EMediaEvent;
 enum class EMediaCacheState;
@@ -309,6 +311,14 @@ public:
 	 */
 	FTimespan GetTime() const;
 
+	/**
+	 * Get the media's current playback time stamp.
+	 *
+	 * @return Playback time stamp.
+	 * @see GetDuration, Seek
+	 */
+	FMediaTimeStamp GetTimeStamp() const;
+
 	/**
 	 * Get the human readable name of the specified track.
 	 *
@@ -662,20 +672,6 @@ public:
 	 */
 	FTimespan GetLastAudioRenderedSampleTime() const;
 
-	/**
-	 * Get time of last audio sample decoded
-	 *
-	 * @return Time of last audio sample decoded.
-	 */
-	FTimespan GetLastAudioSampleProcessedTime() const;
-
-	/**
-	 * Get time of last video sample decoded
-	 *
-	 * @return Time of last video sample decoded.
-	 */
-	FTimespan GetLastVideoSampleProcessedTime() const;
-
 public:
 
 	/** Get an event delegate that is invoked when a media event occurred. */
@@ -753,6 +749,8 @@ protected:
 	void SelectDefaultTracks();
 
 protected:
+	bool HaveAudioPlayback() const;
+	bool HaveVideoPlayback() const;
 
 	/** Fetch audio samples from the player and forward them to the registered sinks. */
 	void ProcessAudioSamples(IMediaSamples& Samples, TRange<FTimespan> TimeRange);
@@ -777,20 +775,27 @@ protected:
 
 private:
 
+	void ProcessVideoSamples(IMediaSamples& Samples, const TRange<FMediaTimeStamp> & TimeRange);
+	void ProcessCaptionSamples(IMediaSamples& Samples, TRange<FMediaTimeStamp> TimeRange);
+	void ProcessSubtitleSamples(IMediaSamples& Samples, TRange<FMediaTimeStamp> TimeRange);
+
+	bool GetCurrentPlaybackTimeRange(TRange<FMediaTimeStamp> & TimeRange, float Rate, FTimespan DeltaTime, bool bDoNotUseFrameStartReference) const;
+
 	/** Audio sample sinks. */
-	TMediaSampleSinks<IMediaAudioSample> AudioSampleSinks;
+	TWeakPtr<FMediaAudioSampleSink, ESPMode::ThreadSafe> PrimaryAudioSink;
+	FMediaAudioSampleSinks AudioSampleSinks;
 
 	/** Caption sample sinks. */
-	TMediaSampleSinks<IMediaOverlaySample> CaptionSampleSinks;
+	FMediaOverlaySampleSinks CaptionSampleSinks;
 
 	/** Metadata sample sinks. */
-	TMediaSampleSinks<IMediaBinarySample> MetadataSampleSinks;
+	FMediaBinarySampleSinks MetadataSampleSinks;
 
 	/** Subtitle sample sinks. */
-	TMediaSampleSinks<IMediaOverlaySample> SubtitleSampleSinks;
+	FMediaOverlaySampleSinks SubtitleSampleSinks;
 
 	/** Video sample sinks. */
-	TMediaSampleSinks<IMediaTextureSample> VideoSampleSinks;
+	FMediaVideoSampleSinks VideoSampleSinks;
 
 private:
 
@@ -809,6 +814,9 @@ private:
 	/** The last used non-zero play rate (zero if playback never started). */
 	float LastRate;
 
+	/** Flag indicating that we have an active audio setup */
+	bool bHaveActiveAudio;
+
 	/** An event delegate that is invoked when a media event occurred. */
 	FOnMediaEvent MediaEvent;
 
@@ -824,12 +832,24 @@ private:
 	/** Media player event queue. */
 	TQueue<EMediaEvent, EQueueMode::Mpsc> QueuedEvents;
 
+	/** CS to make last time values below thread safe */
+	mutable FCriticalSection LastTimeValuesCS;
+
 	/** Time of last audio sample played. */
-	TAtomic<FTimespan> LastAudioRenderedSampleTime;
+	FMediaTimeStampSample LastAudioRenderedSampleTime;
 
 	/** Time of last audio sample decoded. */
-	TAtomic<FTimespan> LastAudioSampleProcessedTime;
+	FMediaTimeStampSample LastAudioSampleProcessedTime;
 
 	/** Time of last video sample decoded. */
-	TAtomic<FTimespan> LastVideoSampleProcessedTime;
+	FMediaTimeStampSample LastVideoSampleProcessedTime;
+
+	/** Timestamp for audio considered "current" for this frame (use ONLY for return to outside code) */
+	FMediaTimeStamp CurrentFrameAudioTimeStamp;
+
+	/** Estimation for next frame's video timestamp (used when no audio present or active in stream) */
+	FMediaTimeStamp NextEstVideoTimeAtFrameStart;
+
+	/** Mediamodule we are working in */
+	IMediaModule* MediaModule;
 };
diff --git a/Engine/Source/Runtime/MediaUtils/Public/MediaSampleCache.h b/Engine/Source/Runtime/MediaUtils/Public/MediaSampleCache.h
--- a/Engine/Source/Runtime/MediaUtils/Public/MediaSampleCache.h
+++ b/Engine/Source/Runtime/MediaUtils/Public/MediaSampleCache.h
@@ -30,7 +30,7 @@ class MEDIAUTILS_API FMediaSampleCache
 	{
 		static FORCEINLINE FTimespan GetSetKey(const TSharedPtr<SampleType, ESPMode::ThreadSafe>& Element)
 		{
-			return Element->GetTime();
+			return Element->GetTime().Time;
 		}
 
 		static FORCEINLINE bool Matches(FTimespan A, FTimespan B)
diff --git a/Engine/Source/Runtime/MediaUtils/Public/MediaSampleQueue.h b/Engine/Source/Runtime/MediaUtils/Public/MediaSampleQueue.h
--- a/Engine/Source/Runtime/MediaUtils/Public/MediaSampleQueue.h
+++ b/Engine/Source/Runtime/MediaUtils/Public/MediaSampleQueue.h
@@ -15,22 +15,21 @@
 #include "MediaSampleSink.h"
 #include "MediaSampleSource.h"
 
+#include "IMediaTimeSource.h"
 #include "IMediaTextureSample.h"
 
 /**
  * Template for media sample queues.
  */
-template<typename SampleType>
+template<typename SampleType, typename SinkType=TMediaSampleSink<SampleType>>
 class TMediaSampleQueue
-	: public TMediaSampleSink<SampleType>
+	: public SinkType
 	, public TMediaSampleSource<SampleType>
 {
 public:
 
 	/** Default constructor. */
 	TMediaSampleQueue()
-		: NumSamples(0)
-		, PendingFlushes(0)
 	{ }
 
 	/** Virtual destructor. */
@@ -50,7 +49,7 @@ public:
 	 */
 	int32 Num() const
 	{
-		return NumSamples;
+		return Samples.Num();
 	}
 
 public:
@@ -59,24 +58,21 @@ public:
 
 	virtual bool Dequeue(TSharedPtr<SampleType, ESPMode::ThreadSafe>& OutSample) override
 	{
-		DoPendingFlushes();
+		FScopeLock Lock(&CriticalSection);
 
-		TSharedPtr<SampleType, ESPMode::ThreadSafe> Sample;
-
-		if (!Samples.Peek(Sample))
+		if (Samples.Num() == 0)
 		{
 			return false; // empty queue
 		}
+
+		TSharedPtr<SampleType, ESPMode::ThreadSafe> Sample(Samples[0]);
 			
 		if (!Sample.IsValid())
 		{
 			return false; // pending flush
 		}
 
-		Samples.Pop();
-
-		FPlatformAtomics::InterlockedDecrement(&NumSamples);
-		check(NumSamples >= 0);
+		Samples.RemoveAt(0);
 
 		OutSample = Sample;
 
@@ -85,15 +81,15 @@ public:
 
 	virtual bool Peek(TSharedPtr<SampleType, ESPMode::ThreadSafe>& OutSample) override
 	{
-		DoPendingFlushes();
+		FScopeLock Lock(&CriticalSection);
 
-		TSharedPtr<SampleType, ESPMode::ThreadSafe> Sample;
-
-		if (!Samples.Peek(Sample))
+		if (Samples.Num() == 0)
 		{
 			return false; // empty queue
 		}
 
+		TSharedPtr<SampleType, ESPMode::ThreadSafe> Sample(Samples[0]);
+
 		if (!Sample.IsValid())
 		{
 			return false; // pending flush
@@ -106,86 +102,237 @@ public:
 
 	virtual bool Pop() override
 	{
-		TSharedPtr<SampleType, ESPMode::ThreadSafe> Sample;
+		FScopeLock Lock(&CriticalSection);
 
-		if (!Samples.Peek(Sample))
+		if (Samples.Num() == 0)
 		{
 			return false; // empty queue
 		}
 
-		if (!Sample.IsValid())
+		if (!Samples[0].IsValid())
 		{
 			return false; // pending flush
 		}
 
-		Samples.Pop();
-
-		FPlatformAtomics::InterlockedDecrement(&NumSamples);
-		check(NumSamples >= 0);
+		Samples.RemoveAt(0);
 
 		return true;
 	}
 
-public:
+	bool FetchBestSampleForTimeRange(const TRange<FMediaTimeStamp> & TimeRange, TSharedPtr<SampleType, ESPMode::ThreadSafe>& OutSample, bool bReverse)
+	{
+		// Code below assumes a fully specified range, no open bounds!
+		check(TimeRange.HasLowerBound() && TimeRange.HasUpperBound());
 
-	//~ TMediaSampleSink interface (to be called only from producer thread)
+		FScopeLock Lock(&CriticalSection);
 
-	virtual bool Enqueue(const TSharedRef<SampleType, ESPMode::ThreadSafe>& Sample) override
-	{
-		FPlatformAtomics::InterlockedIncrement(&NumSamples); // avoid negative sample count in Dequeue
+		int32 Num = Samples.Num();
+		if (Num == 0)
+		{
+			return false;
+		}
 
-		if (!Samples.Enqueue(Sample))
+		int32 FirstPossibleIndex = -1;
+		int32 LastPossibleIndex = -1;
+		int32 FirstKeepIndex = 0;
+		for (int32 Idx = 0; Idx < Num; ++Idx)
 		{
-			FPlatformAtomics::InterlockedDecrement(&NumSamples);
+			const TSharedPtr<SampleType, ESPMode::ThreadSafe> & Sample = Samples[Idx];
+			TRange<FMediaTimeStamp> SampleTimeRange = !bReverse ? TRange<FMediaTimeStamp>(Sample->GetTime(), Sample->GetTime() + Sample->GetDuration())
+																: TRange<FMediaTimeStamp>(Sample->GetTime() - Sample->GetDuration(), Sample->GetTime());
 
-			return false;
+			if (TimeRange.Overlaps(SampleTimeRange))
+			{
+				// Sample is at least partially inside the requested range, recall the range of samples we find...
+				if (FirstPossibleIndex < 0)
+				{
+					FirstPossibleIndex = Idx;
+				}
+				LastPossibleIndex = Idx;
+			}
+			else
+			{
+				if (!bReverse ? (SampleTimeRange.GetLowerBoundValue() >= TimeRange.GetUpperBoundValue()) :
+								(SampleTimeRange.GetUpperBoundValue() <= TimeRange.GetLowerBoundValue()))
+				{
+					// Sample is entirely past requested time range, we can stop
+					break;
+				}
+
+				// We must not get here if we already found overlapping samples
+				check(FirstPossibleIndex < 0);
+
+				// Sample is before time range, we will delete is later, no reason to keep it
+				++FirstKeepIndex;
+			}
 		}
 
-		return true;
-	}
+		// Found anything?
+		if (FirstPossibleIndex >= 0)
+		{
+			if (FirstPossibleIndex != LastPossibleIndex)
+			{
+				// More then one sample. Find the one that fits the bill, best...
+				// (we look for the one with the largest overlap & newest time)
+				FMediaTimeStamp BestDuration(FTimespan::Zero());
+				int32 BestIndex = FirstPossibleIndex;
+				for (int32 Idx = FirstPossibleIndex; Idx <= LastPossibleIndex; ++Idx)
+				{
+					const TSharedPtr<SampleType, ESPMode::ThreadSafe> & Sample = Samples[Idx];
+					TRange<FMediaTimeStamp> SampleTimeRange = !bReverse ? TRange<FMediaTimeStamp>(Sample->GetTime(), Sample->GetTime() + Sample->GetDuration())
+						: TRange<FMediaTimeStamp>(Sample->GetTime() - Sample->GetDuration(), Sample->GetTime());
+					TRange<FMediaTimeStamp> SampleInRangeRange(TRange<FMediaTimeStamp>::Intersection(SampleTimeRange, TimeRange));
+
+					FMediaTimeStamp SampleDuration(SampleInRangeRange.Size<FMediaTimeStamp>());
+					if (SampleDuration >= BestDuration)
+					{
+						BestDuration = SampleDuration;
+						BestIndex = Idx;
+					}
+				}
+
+				check(BestIndex >= FirstKeepIndex);
+
+				// Found the best. Return it & delete all candidate samples up and including it from the queue
+				OutSample = Samples[BestIndex];
+				Samples.RemoveAt(FirstPossibleIndex, BestIndex - FirstPossibleIndex + 1);
+			}
+			else
+			{
+				// Single sample found: we just take it!
+				OutSample = Samples[FirstPossibleIndex];
+				Samples.RemoveAt(FirstPossibleIndex);
+			}
+		}
 
-	virtual void RequestFlush() override
-	{
-		Samples.Enqueue(nullptr); // insert flush marker
-		FPlatformAtomics::InterlockedIncrement(&PendingFlushes);
+		// Cleanup samples that are now considered outdated...
+		if (FirstKeepIndex != 0)
+		{
+			Samples.RemoveAt(0, FirstKeepIndex);
+		}
+
+		// Return true if we got a sample...
+		return (FirstPossibleIndex >= 0);
 	}
 
-protected:
 
-	/** Perform any pending flushes. */
-	void DoPendingFlushes()
+	uint32 PurgeOutdatedSamples(const FMediaTimeStamp & ReferenceTime, bool bReversed)
 	{
-		TSharedPtr<SampleType, ESPMode::ThreadSafe> Sample;
+		FScopeLock Lock(&CriticalSection);
 
-		while ((PendingFlushes > 0) && Samples.Dequeue(Sample))
+		int32 Num = Samples.Num();
+		if (Num > 0)
 		{
-			if (Sample.IsValid())
+			int32 Idx = 0;
+			if (!bReversed)
 			{
-				FPlatformAtomics::InterlockedDecrement(&NumSamples);
-				check(NumSamples >= 0);
+				for (; Idx < Num; ++Idx)
+				{
+					const TSharedPtr<SampleType, ESPMode::ThreadSafe> & Sample = Samples[Idx];
+					if ((Sample->GetTime().Time + Sample->GetDuration()) > ReferenceTime.Time)
+					{
+						break;
+					}
+				}
 			}
 			else
 			{
-				FPlatformAtomics::InterlockedDecrement(&PendingFlushes);
+				for (; Idx < Num; ++Idx)
+				{
+					const TSharedPtr<SampleType, ESPMode::ThreadSafe> & Sample = Samples[Idx];
+					if ((Sample->GetTime().Time - Sample->GetDuration()) < ReferenceTime.Time)
+					{
+						break;
+					}
+				}
 			}
+			if (Idx > 0)
+			{
+				Samples.RemoveAt(0, Idx);
+			}
+			return Idx;
 		}
+
+		return 0;
 	}
 
-protected:
+public:
 
-	/** Number of samples in the queue. */
-	int32 NumSamples;
+	//~ TMediaSampleSink interface (to be called only from producer thread)
+
+	virtual bool Enqueue(const TSharedRef<SampleType, ESPMode::ThreadSafe>& Sample) override
+	{
+		FScopeLock Lock(&CriticalSection);
+
+		Samples.Push(Sample);
+		return true;
+	}
 
-	/** Number of pending flushes. */
-	int32 PendingFlushes;
+	virtual void RequestFlush() override
+	{
+		FScopeLock Lock(&CriticalSection);
+		Samples.Empty();
+	}
 
-	/** Audio sample queue. */
-	TQueue<TSharedPtr<SampleType, ESPMode::ThreadSafe>, EQueueMode::Mpsc> Samples;
+protected:
+	mutable FCriticalSection CriticalSection;
+	TArray<TSharedPtr<SampleType, ESPMode::ThreadSafe>> Samples;
 };
 
 
-/** Type definition for audio sample queue. */
-typedef TMediaSampleQueue<class IMediaAudioSample> FMediaAudioSampleQueue;
+/** audio sample queue. */
+class FMediaAudioSampleQueue : public TMediaSampleQueue<class IMediaAudioSample, class FMediaAudioSampleSink>
+{
+public:
+	FMediaAudioSampleQueue(int32 InMaxAudioSamplesInQueue = -1)
+		: MaxAudioSamplesInQueue(InMaxAudioSamplesInQueue) {}
+
+	void SetAudioTime(const FMediaTimeStampSample & InAudioTime)
+	{
+		FScopeLock Lock(&CriticalSection);
+		AudioTime = InAudioTime;
+	}
+
+	FMediaTimeStampSample GetAudioTime() const override
+	{
+		FScopeLock Lock(&CriticalSection);
+		return AudioTime;
+	}
+
+	void InvalidateAudioTime() override
+	{
+		FScopeLock Lock(&CriticalSection);
+		AudioTime.Invalidate();
+	}
+
+	virtual void RequestFlush() override
+	{
+		FScopeLock Lock(&CriticalSection);
+		TMediaSampleQueue<class IMediaAudioSample, class FMediaAudioSampleSink>::RequestFlush();
+		AudioTime.Invalidate();
+	}
+
+	virtual bool Enqueue(const TSharedRef<IMediaAudioSample, ESPMode::ThreadSafe>& Sample) override
+	{
+		FScopeLock Lock(&CriticalSection);
+
+		if (MaxAudioSamplesInQueue > 0 && Samples.Num() >= MaxAudioSamplesInQueue)
+		{
+			return false;
+		}
+
+		Samples.Push(Sample);
+		return true;
+	}
+
+	virtual bool CanAcceptSamples(int32 NumSamples) const
+	{
+		return (Samples.Num() + NumSamples) <= MaxAudioSamplesInQueue;
+	}
+private:
+	int32 MaxAudioSamplesInQueue;
+	FMediaTimeStampSample AudioTime;
+};
 
 /** Type definition for binary sample queue. */
 typedef TMediaSampleQueue<class IMediaBinarySample> FMediaBinarySampleQueue;
@@ -195,4 +342,3 @@ typedef TMediaSampleQueue<class IMediaOverlaySample> FMediaOverlaySampleQueue;
 
 /** Type definition for texture sample queue. */
 typedef TMediaSampleQueue<class IMediaTextureSample> FMediaTextureSampleQueue;
-
diff --git a/Engine/Source/Runtime/MediaUtils/Public/MediaSampleSink.h b/Engine/Source/Runtime/MediaUtils/Public/MediaSampleSink.h
--- a/Engine/Source/Runtime/MediaUtils/Public/MediaSampleSink.h
+++ b/Engine/Source/Runtime/MediaUtils/Public/MediaSampleSink.h
@@ -5,7 +5,7 @@
 #include "CoreTypes.h"
 #include "Misc/Timespan.h"
 #include "Templates/SharedPointer.h"
-
+#include "IMediaTimeSource.h"
 
 /**
  * Interface for media sample sinks.
@@ -52,7 +52,17 @@ public:
 
 
 /** Type definition for audio sample sink. */
-typedef TMediaSampleSink<class IMediaAudioSample> FMediaAudioSampleSink;
+class FMediaAudioSampleSink : public TMediaSampleSink<class IMediaAudioSample>
+{
+public:
+	/**
+	 * Get last sampled current audio timestamp being played
+	 */
+	virtual FMediaTimeStampSample GetAudioTime() const = 0;
+	virtual void InvalidateAudioTime() = 0;
+
+	virtual bool CanAcceptSamples(int32 NumSamples) const = 0;
+};
 
 /** Type definition for binary sample sink. */
 typedef TMediaSampleSink<class IMediaBinarySample> FMediaBinarySampleSink;
diff --git a/Engine/Source/Runtime/MediaUtils/Public/MediaSampleSinks.h b/Engine/Source/Runtime/MediaUtils/Public/MediaSampleSinks.h
--- a/Engine/Source/Runtime/MediaUtils/Public/MediaSampleSinks.h
+++ b/Engine/Source/Runtime/MediaUtils/Public/MediaSampleSinks.h
@@ -14,7 +14,7 @@
  *
  * @param SampleType The type of media samples that the sinks process.
  */
-template<typename SampleType>
+template<typename SampleType, typename SinkType=TMediaSampleSink<SampleType>>
 class TMediaSampleSinks
 {
 public:
@@ -25,7 +25,7 @@ public:
 	 * @param SampleSink The sink to add.
 	 * @see Num, Remove
 	 */
-	void Add(const TSharedRef<TMediaSampleSink<SampleType>, ESPMode::ThreadSafe>& SampleSink)
+	void Add(const TSharedRef<SinkType, ESPMode::ThreadSafe>& SampleSink)
 	{
 		Sinks.AddUnique(SampleSink);
 	}
@@ -46,7 +46,7 @@ public:
 
 		for (int32 SinkIndex = Sinks.Num() - 1; SinkIndex >= 0; --SinkIndex)
 		{
-			TSharedPtr<TMediaSampleSink<SampleType>, ESPMode::ThreadSafe> Sink = Sinks[SinkIndex].Pin();
+			TSharedPtr<SinkType, ESPMode::ThreadSafe> Sink = Sinks[SinkIndex].Pin();
 
 			if (Sink.IsValid())
 			{
@@ -79,7 +79,7 @@ public:
 	{
 		for (int32 SinkIndex = Sinks.Num() - 1; SinkIndex >= 0; --SinkIndex)
 		{
-			TSharedPtr<TMediaSampleSink<SampleType>, ESPMode::ThreadSafe> Sink = Sinks[SinkIndex].Pin();
+			TSharedPtr<SinkType, ESPMode::ThreadSafe> Sink = Sinks[SinkIndex].Pin();
 
 			if (Sink.IsValid())
 			{
@@ -109,13 +109,31 @@ public:
 	 * @param SampleSink The sink to remove.
 	 * @see Add, Num
 	 */
-	void Remove(const TSharedRef<TMediaSampleSink<SampleType>, ESPMode::ThreadSafe>& SampleSink)
+	void Remove(const TSharedRef<SinkType, ESPMode::ThreadSafe>& SampleSink)
 	{
 		Sinks.Remove(SampleSink);
 	}
 
-private:
+protected:
 
 	/** The collection of registered sinks. */
-	TArray<TWeakPtr<TMediaSampleSink<SampleType>, ESPMode::ThreadSafe>> Sinks;
+	TArray<TWeakPtr<SinkType, ESPMode::ThreadSafe>> Sinks;
 };
+
+typedef TMediaSampleSinks<IMediaTextureSample, FMediaTextureSampleSink> FMediaVideoSampleSinks;
+
+class FMediaAudioSampleSinks : public TMediaSampleSinks<IMediaAudioSample, FMediaAudioSampleSink>
+{
+public:
+	TSharedPtr<FMediaAudioSampleSink, ESPMode::ThreadSafe> GetPrimaryAudioSink()
+	{
+		if (Sinks.Num() == 0)
+		{
+			return TSharedPtr<FMediaAudioSampleSink, ESPMode::ThreadSafe>();
+		}
+		return Sinks[0].Pin();
+	}
+};
+
+typedef TMediaSampleSinks<IMediaOverlaySample, FMediaOverlaySampleSink> FMediaOverlaySampleSinks;
+typedef TMediaSampleSinks<IMediaBinarySample, FMediaBinarySampleSink> FMediaBinarySampleSinks;
diff --git a/Engine/Source/Runtime/MediaUtils/Public/MediaSamples.h b/Engine/Source/Runtime/MediaUtils/Public/MediaSamples.h
--- a/Engine/Source/Runtime/MediaUtils/Public/MediaSamples.h
+++ b/Engine/Source/Runtime/MediaUtils/Public/MediaSamples.h
@@ -83,7 +83,7 @@ public:
 	 * @return Number of samples.
 	 * @see AddAudio, NumCaption, NumMetadata, NumSubtitle, NumVideo
 	 */
-	int32 NumAudio() const
+	int32 NumAudio() const override
 	{
 		return AudioSampleQueue.Num();
 	}
@@ -94,7 +94,7 @@ public:
 	 * @return Number of samples.
 	 * @see AddCaption, NumAudio, NumMetadata, NumSubtitle, NumVideo
 	 */
-	int32 NumCaption() const
+	int32 NumCaption() const override
 	{
 		return CaptionSampleQueue.Num();
 	}
@@ -105,7 +105,7 @@ public:
 	 * @return Number of samples.
 	 * @see AddMetadata, NumAudio, NumCaption, NumSubtitle, NumVideo
 	 */
-	int32 NumMetadataSamples() const
+	int32 NumMetadataSamples() const override
 	{
 		return MetadataSampleQueue.Num();
 	}
@@ -116,7 +116,7 @@ public:
 	 * @return Number of samples.
 	 * @see AddSubtitle, NumAudio, NumCaption, NumMetadata, NumVideo
 	 */
-	int32 NumSubtitleSamples() const
+	int32 NumSubtitleSamples() const override
 	{
 		return SubtitleSampleQueue.Num();
 	}
@@ -127,11 +127,15 @@ public:
 	 * @return Number of samples.
 	 * @see AddVideo, NumAudio, NumCaption, NumMetadata, NumSubtitle
 	 */
-	int32 NumVideoSamples() const
+	int32 NumVideoSamples() const override
 	{
 		return VideoSampleQueue.Num();
 	}
 
+	/**
+	 * Peek next video sample's timestamp
+	 * @return true if value could be retrieved, false otherwise
+	 */
 public:
 
 	//~ IMediaSamples interface
@@ -143,6 +147,27 @@ public:
 	virtual bool FetchVideo(TRange<FTimespan> TimeRange, TSharedPtr<IMediaTextureSample, ESPMode::ThreadSafe>& OutSample) override;
 	virtual void FlushSamples() override;
 
+	virtual bool FetchCaption(TRange<FMediaTimeStamp> TimeRange, TSharedPtr<IMediaOverlaySample, ESPMode::ThreadSafe>& OutSample) override;
+	virtual bool FetchSubtitle(TRange<FMediaTimeStamp> TimeRange, TSharedPtr<IMediaOverlaySample, ESPMode::ThreadSafe>& OutSample) override;
+	virtual bool FetchVideo(TRange<FMediaTimeStamp> TimeRange, TSharedPtr<IMediaTextureSample, ESPMode::ThreadSafe>& OutSample) override;
+
+	virtual EFetchBestSampleResult FetchBestVideoSampleForTimeRange(const TRange<FMediaTimeStamp> & TimeRange, TSharedPtr<IMediaTextureSample, ESPMode::ThreadSafe>& OutSample, bool bReverse) override;
+
+	virtual bool PeekVideoSampleTime(FMediaTimeStamp & TimeStamp) override
+	{
+		TSharedPtr<IMediaTextureSample, ESPMode::ThreadSafe> Sample;
+		if (!VideoSampleQueue.Peek(Sample))
+		{
+			return false;
+		}
+		TimeStamp = Sample->GetTime();
+		return true;
+	}
+
+	virtual uint32 PurgeOutdatedVideoSamples(const FMediaTimeStamp & ReferenceTime, bool bReversed) override;
+
+	virtual bool CanReceiveVideoSamples(uint32 Num) const override;
+
 private:
 
 	/** Audio sample queue. */
diff --git a/Engine/Source/Runtime/RenderCore/Private/MediaShaders.cpp b/Engine/Source/Runtime/RenderCore/Private/MediaShaders.cpp
--- a/Engine/Source/Runtime/RenderCore/Private/MediaShaders.cpp
+++ b/Engine/Source/Runtime/RenderCore/Private/MediaShaders.cpp
@@ -153,7 +153,8 @@ SHADER_PARAMETER(FMatrix, ColorTransform)
 SHADER_PARAMETER(uint32, OutputWidth)
 SHADER_PARAMETER(uint32, SrgbToLinear)
 SHADER_PARAMETER(FVector2D, UVScale)
-SHADER_PARAMETER_TEXTURE(Texture2D, Texture)
+SHADER_PARAMETER_SRV(Texture2D, SRV_Y)
+SHADER_PARAMETER_SRV(Texture2D, SRV_UV)
 SHADER_PARAMETER_SAMPLER(SamplerState, SamplerB)
 SHADER_PARAMETER_SAMPLER(SamplerState, SamplerP)
 END_GLOBAL_SHADER_PARAMETER_STRUCT()
@@ -171,7 +172,8 @@ void FNV12ConvertPS::SetParameters(FRHICommandList& CommandList, TRefCountPtr<FR
 		UB.SamplerB = TStaticSamplerState<SF_Bilinear>::GetRHI();
 		UB.SamplerP = TStaticSamplerState<SF_Point>::GetRHI();
 		UB.SrgbToLinear = SrgbToLinear;
-		UB.Texture = NV12Texture;
+		UB.SRV_Y = RHICreateShaderResourceView(NV12Texture, 0, 1, PF_G8);
+		UB.SRV_UV = RHICreateShaderResourceView(NV12Texture, 0, 1, PF_R8G8);
 		UB.UVScale = FVector2D((float)OutputDimensions.X / (float)NV12Texture->GetSizeX(), (float)OutputDimensions.Y / (float)NV12Texture->GetSizeY());
 	}
 
@@ -179,6 +181,56 @@ void FNV12ConvertPS::SetParameters(FRHICommandList& CommandList, TRefCountPtr<FR
 	SetUniformBufferParameter(CommandList, CommandList.GetBoundPixelShader(), GetUniformBufferParameter<FNV12ConvertUB>(), Data);
 }
 
+void FNV12ConvertPS::SetParameters(FRHICommandList& CommandList, const FIntPoint & TexDim, FShaderResourceViewRHIRef SRV_Y, FShaderResourceViewRHIRef SRV_UV, const FIntPoint& OutputDimensions, const FMatrix& ColorTransform, const FVector& YUVOffset, bool SrgbToLinear)
+{
+	FNV12ConvertUB UB;
+	{
+		UB.ColorTransform = MediaShaders::CombineColorTransformAndOffset(ColorTransform, YUVOffset);
+		UB.OutputWidth = OutputDimensions.X;
+		UB.SamplerB = TStaticSamplerState<SF_Bilinear>::GetRHI();
+		UB.SamplerP = TStaticSamplerState<SF_Point>::GetRHI();
+		UB.SrgbToLinear = SrgbToLinear;
+		UB.SRV_Y = SRV_Y;
+		UB.SRV_UV = SRV_UV;
+		UB.UVScale = FVector2D((float)OutputDimensions.X / (float)TexDim.X, (float)OutputDimensions.Y / (float)TexDim.Y);
+	}
+
+	TUniformBufferRef<FNV12ConvertUB> Data = TUniformBufferRef<FNV12ConvertUB>::CreateUniformBufferImmediate(UB, UniformBuffer_SingleFrame);
+	SetUniformBufferParameter(CommandList, CommandList.GetBoundPixelShader(), GetUniformBufferParameter<FNV12ConvertUB>(), Data);
+}
+
+/* FNV12ConvertAsBytesPS shader
+ *****************************************************************************/
+
+BEGIN_GLOBAL_SHADER_PARAMETER_STRUCT(FNV12ConvertAsBytesUB, )
+SHADER_PARAMETER(FMatrix, ColorTransform)
+SHADER_PARAMETER(uint32, OutputWidth)
+SHADER_PARAMETER(uint32, SrgbToLinear)
+SHADER_PARAMETER(FVector2D, UVScale)
+SHADER_PARAMETER_TEXTURE(Texture2D, Texture)
+SHADER_PARAMETER_SAMPLER(SamplerState, SamplerB)
+SHADER_PARAMETER_SAMPLER(SamplerState, SamplerP)
+END_GLOBAL_SHADER_PARAMETER_STRUCT()
+
+IMPLEMENT_GLOBAL_SHADER_PARAMETER_STRUCT(FNV12ConvertAsBytesUB, "NV12ConvertAsBytesUB");
+IMPLEMENT_SHADER_TYPE(, FNV12ConvertAsBytesPS, TEXT("/Engine/Private/MediaShaders.usf"), TEXT("NV12ConvertAsBytesPS"), SF_Pixel);
+
+void FNV12ConvertAsBytesPS::SetParameters(FRHICommandList& CommandList, TRefCountPtr<FRHITexture2D> NV12Texture, const FIntPoint& OutputDimensions, const FMatrix& ColorTransform, const FVector& YUVOffset, bool SrgbToLinear)
+{
+	FNV12ConvertAsBytesUB UB;
+	{
+		UB.ColorTransform = MediaShaders::CombineColorTransformAndOffset(ColorTransform, YUVOffset);
+		UB.OutputWidth = OutputDimensions.X;
+		UB.SamplerB = TStaticSamplerState<SF_Bilinear>::GetRHI();
+		UB.SamplerP = TStaticSamplerState<SF_Point>::GetRHI();
+		UB.SrgbToLinear = SrgbToLinear;
+		UB.Texture = NV12Texture;
+		UB.UVScale = FVector2D((float)OutputDimensions.X / (float)NV12Texture->GetSizeX(), (float)OutputDimensions.Y / (float)NV12Texture->GetSizeY());
+	}
+
+	TUniformBufferRef<FNV12ConvertAsBytesUB> Data = TUniformBufferRef<FNV12ConvertAsBytesUB>::CreateUniformBufferImmediate(UB, UniformBuffer_SingleFrame);
+	SetUniformBufferParameter(CommandList, CommandList.GetBoundPixelShader(), GetUniformBufferParameter<FNV12ConvertAsBytesUB>(), Data);
+}
 
 /* FNV21ConvertPS shader
  *****************************************************************************/
diff --git a/Engine/Source/Runtime/RenderCore/Public/MediaShaders.h b/Engine/Source/Runtime/RenderCore/Public/MediaShaders.h
--- a/Engine/Source/Runtime/RenderCore/Public/MediaShaders.h
+++ b/Engine/Source/Runtime/RenderCore/Public/MediaShaders.h
@@ -194,7 +194,7 @@ public:
 
 
 /**
- * Pixel shader to convert a NV12 frame to RGBA.
+ * Pixel shader to convert a NV12 frame to RGBA. (from NV12 texture)
  *
  * This shader expects an NV12 frame packed into a single texture in PF_G8 format.
  *
@@ -218,6 +218,36 @@ public:
 		: FGlobalShader(Initializer)
 	{ }
 
+	RENDERCORE_API void SetParameters(FRHICommandList& RHICmdList, TRefCountPtr<FRHITexture2D> NV12Texture, const FIntPoint& OutputDimensions, const FMatrix& ColorTransform, const FVector& YUVOffset, bool SrgbToLinear);
+	RENDERCORE_API void SetParameters(FRHICommandList& RHICmdList, const FIntPoint & TexDim, FShaderResourceViewRHIRef SRV_Y, FShaderResourceViewRHIRef SRV_UV, const FIntPoint& OutputDimensions, const FMatrix& ColorTransform, const FVector& YUVOffset, bool SrgbToLinear);
+};
+
+
+/**
+ * Pixel shader to convert a NV12 frame to RGBA (NV12 data; texture viewed as G8)
+ *
+ * This shader expects an NV12 frame packed into a single texture in PF_G8 format.
+ *
+ * @see http://www.fourcc.org/yuv.php#NV12
+ */
+class FNV12ConvertAsBytesPS
+	: public FGlobalShader
+{
+	DECLARE_EXPORTED_SHADER_TYPE(FNV12ConvertPS, Global, RENDERCORE_API);
+
+public:
+
+	static bool ShouldCompilePermutation(const FGlobalShaderPermutationParameters& Parameters)
+	{
+		return IsFeatureLevelSupported(Parameters.Platform, ERHIFeatureLevel::ES3_1);
+	}
+
+	FNV12ConvertAsBytesPS() { }
+
+	FNV12ConvertAsBytesPS(const ShaderMetaType::CompiledShaderInitializerType& Initializer)
+		: FGlobalShader(Initializer)
+	{ }
+
 	RENDERCORE_API void SetParameters(FRHICommandList& RHICmdList, TRefCountPtr<FRHITexture2D> NV12Texture, const FIntPoint& OutputDimensions, const FMatrix& ColorTransform, const FVector& YUVOffset, bool SrgbToLinear);
 };
 
