Disabling PVS Studio warning
#rnx

diff --git a/Engine/Source/Runtime/MediaUtils/Public/MediaSampleQueue.h b/Engine/Source/Runtime/MediaUtils/Public/MediaSampleQueue.h
--- a/Engine/Source/Runtime/MediaUtils/Public/MediaSampleQueue.h
+++ b/Engine/Source/Runtime/MediaUtils/Public/MediaSampleQueue.h
@@ -200,7 +200,7 @@ public:
 			else
 			{
 				// Single sample found: we just take it!
-				OutSample = Samples[FirstPossibleIndex];
+				OutSample = Samples[FirstPossibleIndex]; //-V781 PVS-Studio triggers incorrectly here: Variable checked after being used (likely a template code issue, but harmless)
 				Samples.RemoveAt(FirstPossibleIndex);
 			}
 		}
