Slight change to logic selection the current frame for MediaFramework (of nothing fitting is available, slightly too old frames are allowed to be shown)
#rnx

diff --git a/Engine/Source/Runtime/MediaUtils/Public/MediaSampleQueue.h b/Engine/Source/Runtime/MediaUtils/Public/MediaSampleQueue.h
--- a/Engine/Source/Runtime/MediaUtils/Public/MediaSampleQueue.h
+++ b/Engine/Source/Runtime/MediaUtils/Public/MediaSampleQueue.h
@@ -124,6 +124,8 @@ public:
 		// Code below assumes a fully specified range, no open bounds!
 		check(TimeRange.HasLowerBound() && TimeRange.HasUpperBound());
 
+		OutSample.Reset();
+
 		FScopeLock Lock(&CriticalSection);
 
 		int32 Num = Samples.Num();
@@ -205,14 +207,20 @@ public:
 			}
 		}
 
-		// Cleanup samples that are now considered outdated...
+		// Any frames considered outdated?
 		if (FirstKeepIndex != 0)
 		{
+			// In case we got no new frame that fits into the current frame, return the newest of the "old" frames
+			if (!OutSample.IsValid())
+			{
+				OutSample = Samples[FirstKeepIndex - 1];
+			}
+			// Cleanup samples that are now considered outdated...
 			Samples.RemoveAt(0, FirstKeepIndex);
 		}
 
 		// Return true if we got a sample...
-		return (FirstPossibleIndex >= 0);
+		return (OutSample.IsValid());
 	}
 
 
