Electra: added support for elementary streams of different duration (to allow for audio portion being longer than video)

#rnx

diff --git a/Engine/Source/Runtime/Media/Public/IMediaEventSink.h b/Engine/Source/Runtime/Media/Public/IMediaEventSink.h
--- a/Engine/Source/Runtime/Media/Public/IMediaEventSink.h
+++ b/Engine/Source/Runtime/Media/Public/IMediaEventSink.h
@@ -45,6 +45,12 @@ enum class EMediaEvent
 
 	/** Player requests purge of outdated video samples from queues. */
 	Internal_PurgeVideoSamplesHint = Internal_Start,
+
+	/** Player notifies about availability or unavailablity of media samples. **/
+	Internal_VideoSamplesAvailable,
+	Internal_VideoSamplesUnavailable,
+	Internal_AudioSamplesAvailable,
+	Internal_AudioSamplesUnavailable
 };
 
 
diff --git a/Engine/Source/Runtime/MediaUtils/Private/MediaHelpers.cpp b/Engine/Source/Runtime/MediaUtils/Private/MediaHelpers.cpp
--- a/Engine/Source/Runtime/MediaUtils/Private/MediaHelpers.cpp
+++ b/Engine/Source/Runtime/MediaUtils/Private/MediaHelpers.cpp
@@ -25,6 +25,11 @@ namespace MediaUtils
 		case EMediaEvent::TracksChanged: return TEXT("TracksChanged");
 
 		case EMediaEvent::Internal_PurgeVideoSamplesHint: return TEXT("Internal: PurgeVideoSamplesHint");
+		case EMediaEvent::Internal_VideoSamplesAvailable: return TEXT("Internal: VideoSamplesAvailable");
+		case EMediaEvent::Internal_VideoSamplesUnavailable: return TEXT("Internal: VideoSamplesUnavailable");
+		case EMediaEvent::Internal_AudioSamplesAvailable: return TEXT("Internal: AudioSamplesAvailable");
+		case EMediaEvent::Internal_AudioSamplesUnavailable: return TEXT("Internal: AudioSamplesUnavailable");
+
 
 		// no default case; all cases must be implemented!
 		}
diff --git a/Engine/Source/Runtime/MediaUtils/Private/MediaPlayerFacade.cpp b/Engine/Source/Runtime/MediaUtils/Private/MediaPlayerFacade.cpp
--- a/Engine/Source/Runtime/MediaUtils/Private/MediaPlayerFacade.cpp
+++ b/Engine/Source/Runtime/MediaUtils/Private/MediaPlayerFacade.cpp
@@ -94,7 +94,9 @@ FMediaPlayerFacade::FMediaPlayerFacade()
 	, Cache(new FMediaSampleCache)
 	, LastRate(0.0f)
 	, bHaveActiveAudio(false)
-{ 
+	, VideoSampleAvailability(-1)
+	, AudioSampleAvailability(-1)
+{
 	MediaModule = FModuleManager::LoadModulePtr<IMediaModule>("Media");
 }
 
@@ -129,7 +131,7 @@ void FMediaPlayerFacade::AddCaptionSampleSink(const TSharedRef<FMediaOverlaySamp
 {
 	CaptionSampleSinks.Add(SampleSink);
 }
-	
+
 
 void FMediaPlayerFacade::AddMetadataSampleSink(const TSharedRef<FMediaBinarySampleSink, ESPMode::ThreadSafe>& SampleSink)
 {
@@ -152,7 +154,12 @@ void FMediaPlayerFacade::AddVideoSampleSink(const TSharedRef<FMediaTextureSample
 
 bool FMediaPlayerFacade::CanPause() const
 {
-	return Player.IsValid() && Player->GetControls().CanControl(EMediaControl::Pause);
+	TSharedPtr<IMediaPlayer, ESPMode::ThreadSafe> CurrentPlayer(Player);
+	if (!CurrentPlayer.IsValid())
+	{
+		return false;
+	}
+	return CurrentPlayer->GetControls().CanControl(EMediaControl::Pause);
 }
 
 
@@ -180,19 +187,34 @@ bool FMediaPlayerFacade::CanPlayUrl(const FString& Url, const IMediaOptions* Opt
 
 bool FMediaPlayerFacade::CanResume() const
 {
-	return Player.IsValid() && Player->GetControls().CanControl(EMediaControl::Resume);
+	TSharedPtr<IMediaPlayer, ESPMode::ThreadSafe> CurrentPlayer(Player);
+	if (!CurrentPlayer.IsValid())
+	{
+		return false;
+	}
+	return CurrentPlayer->GetControls().CanControl(EMediaControl::Resume);
 }
 
 
 bool FMediaPlayerFacade::CanScrub() const
 {
-	return Player.IsValid() && Player->GetControls().CanControl(EMediaControl::Scrub);
+	TSharedPtr<IMediaPlayer, ESPMode::ThreadSafe> CurrentPlayer(Player);
+	if (!CurrentPlayer.IsValid())
+	{
+		return false;
+	}
+	return CurrentPlayer->GetControls().CanControl(EMediaControl::Scrub);
 }
 
 
 bool FMediaPlayerFacade::CanSeek() const
 {
-	return Player.IsValid() && Player->GetControls().CanControl(EMediaControl::Seek);
+	TSharedPtr<IMediaPlayer, ESPMode::ThreadSafe> CurrentPlayer(Player);
+	if (!CurrentPlayer.IsValid())
+	{
+		return false;
+	}
+	return CurrentPlayer->GetControls().CanControl(EMediaControl::Seek);
 }
 
 
@@ -205,10 +227,11 @@ void FMediaPlayerFacade::Close()
 		return;
 	}
 
-	if (Player.IsValid())
+	TSharedPtr<IMediaPlayer, ESPMode::ThreadSafe> CurrentPlayer(Player);
+	if (CurrentPlayer.IsValid())
 	{
 		FScopeLock Lock(&CriticalSection);
-		Player->Close();
+		CurrentPlayer->Close();
 	}
 
 	BlockOnTime = FTimespan::MinValue();
@@ -216,6 +239,10 @@ void FMediaPlayerFacade::Close()
 	CurrentUrl.Empty();
 	LastRate = 0.0f;
 
+	bHaveActiveAudio = false;
+	VideoSampleAvailability = -1;
+	AudioSampleAvailability = -1;
+
 	FlushSinks();
 }
 
@@ -243,7 +270,12 @@ FString FMediaPlayerFacade::GetAudioTrackType(int32 TrackIndex, int32 FormatInde
 
 FTimespan FMediaPlayerFacade::GetDuration() const
 {
-	return Player.IsValid() ? Player->GetControls().GetDuration() : FTimespan::Zero();
+	TSharedPtr<IMediaPlayer, ESPMode::ThreadSafe> CurrentPlayer(Player);
+	if (!CurrentPlayer.IsValid())
+	{
+		return FTimespan::Zero();
+	}
+	return CurrentPlayer->GetControls().GetDuration();
 }
 
 
@@ -255,49 +287,89 @@ const FGuid& FMediaPlayerFacade::GetGuid()
 
 FString FMediaPlayerFacade::GetInfo() const
 {
-	return Player.IsValid() ? Player->GetInfo() : FString();
+	TSharedPtr<IMediaPlayer, ESPMode::ThreadSafe> CurrentPlayer(Player);
+	if (!CurrentPlayer.IsValid())
+	{
+		return FString();
+	}
+	return CurrentPlayer->GetInfo();
 }
 
 
 FText FMediaPlayerFacade::GetMediaName() const
 {
-	return Player.IsValid() ? Player->GetMediaName() : FText::GetEmpty();
+	TSharedPtr<IMediaPlayer, ESPMode::ThreadSafe> CurrentPlayer(Player);
+	if (!CurrentPlayer.IsValid())
+	{
+		return FText::GetEmpty();
+	}
+	return CurrentPlayer->GetMediaName();
 }
 
 
 int32 FMediaPlayerFacade::GetNumTracks(EMediaTrackType TrackType) const
 {
-	return Player.IsValid() ? Player->GetTracks().GetNumTracks(TrackType) : 0;
+	TSharedPtr<IMediaPlayer, ESPMode::ThreadSafe> CurrentPlayer(Player);
+	if (!CurrentPlayer.IsValid())
+	{
+		return 0;
+	}
+	return CurrentPlayer->GetTracks().GetNumTracks(TrackType);
 }
 
 
 int32 FMediaPlayerFacade::GetNumTrackFormats(EMediaTrackType TrackType, int32 TrackIndex) const
 {
-	return Player.IsValid() ? Player->GetTracks().GetNumTrackFormats(TrackType, TrackIndex) : 0;
+	TSharedPtr<IMediaPlayer, ESPMode::ThreadSafe> CurrentPlayer(Player);
+	if (!CurrentPlayer.IsValid())
+	{
+		return 0;
+	}
+	return CurrentPlayer->GetTracks().GetNumTrackFormats(TrackType, TrackIndex);
 }
 
 
 FName FMediaPlayerFacade::GetPlayerName() const
 {
-	return Player.IsValid() ? Player->GetPlayerName() : NAME_None;
+	TSharedPtr<IMediaPlayer, ESPMode::ThreadSafe> CurrentPlayer(Player);
+	if (!CurrentPlayer.IsValid())
+	{
+		return NAME_None;
+	}
+	return CurrentPlayer->GetPlayerName();
 }
 
 
 float FMediaPlayerFacade::GetRate() const
 {
-	return Player.IsValid() ? Player->GetControls().GetRate() : 0.0f;
+	TSharedPtr<IMediaPlayer, ESPMode::ThreadSafe> CurrentPlayer(Player);
+	if (!CurrentPlayer.IsValid())
+	{
+		return 0.0f;
+	}
+	return CurrentPlayer->GetControls().GetRate();
 }
 
 
 int32 FMediaPlayerFacade::GetSelectedTrack(EMediaTrackType TrackType) const
 {
-	return Player.IsValid() ? Player->GetTracks().GetSelectedTrack((EMediaTrackType)TrackType) : INDEX_NONE;
+	TSharedPtr<IMediaPlayer, ESPMode::ThreadSafe> CurrentPlayer(Player);
+	if (!CurrentPlayer.IsValid())
+	{
+		return INDEX_NONE;
+	}
+	return CurrentPlayer->GetTracks().GetSelectedTrack((EMediaTrackType)TrackType);
 }
 
 
 FString FMediaPlayerFacade::GetStats() const
 {
-	return Player.IsValid() ? Player->GetStats() : FString();
+	TSharedPtr<IMediaPlayer, ESPMode::ThreadSafe> CurrentPlayer(Player);
+	if (!CurrentPlayer.IsValid())
+	{
+		return FString();
+	}
+	return CurrentPlayer->GetStats();
 }
 
 
@@ -305,7 +377,12 @@ TRangeSet<float> FMediaPlayerFacade::GetSupportedRates(bool Unthinned) const
 {
 	const EMediaRateThinning Thinning = Unthinned ? EMediaRateThinning::Unthinned : EMediaRateThinning::Thinned;
 
-	return Player.IsValid() ? Player->GetControls().GetSupportedRates(Thinning) : TRangeSet<float>();
+	TSharedPtr<IMediaPlayer, ESPMode::ThreadSafe> CurrentPlayer(Player);
+	if (!CurrentPlayer.IsValid())
+	{
+		return TRangeSet<float>();
+	}
+	return CurrentPlayer->GetControls().GetSupportedRates(Thinning);
 }
 
 
@@ -323,14 +400,16 @@ bool FMediaPlayerFacade::HaveAudioPlayback() const
 
 FTimespan FMediaPlayerFacade::GetTime() const
 {
-	if (!Player.IsValid())
+	TSharedPtr<IMediaPlayer, ESPMode::ThreadSafe> CurrentPlayer(Player);
+
+	if (!CurrentPlayer.IsValid())
 	{
 		return FTimespan::Zero(); // no media opened
 	}
 
 	FTimespan Result;
 
-	if (Player->GetPlayerFeatureFlag(IMediaPlayer::EFeatureFlag::UsePlaybackTimingV2))
+	if (CurrentPlayer->GetPlayerFeatureFlag(IMediaPlayer::EFeatureFlag::UsePlaybackTimingV2))
 	{
 		// New style: framework controls timing - we use GetTimeStamp() and return the legacy part of the value
 		FMediaTimeStamp TimeStamp = GetTimeStamp();
@@ -339,7 +418,7 @@ FTimespan FMediaPlayerFacade::GetTime() const
 	else
 	{
 		// Old style: ask the player for timing
-		Result = Player->GetControls().GetTime() - TimeDelay;
+		Result = CurrentPlayer->GetControls().GetTime() - TimeDelay;
 		if (Result.GetTicks() < 0)
 		{
 			Result = FTimespan::Zero();
@@ -352,12 +431,24 @@ FTimespan FMediaPlayerFacade::GetTime() const
 
 FMediaTimeStamp FMediaPlayerFacade::GetTimeStamp() const
 {
-	if (!Player.IsValid() || !Player->GetPlayerFeatureFlag(IMediaPlayer::EFeatureFlag::UsePlaybackTimingV2))
+	TSharedPtr<IMediaPlayer, ESPMode::ThreadSafe> CurrentPlayer(Player);
+	if (!CurrentPlayer.IsValid())
+	{
+		return FMediaTimeStamp();
+	}
+
+	FScopeLock Lock(&LastTimeValuesCS);
+
+	if (!CurrentPlayer->GetPlayerFeatureFlag(IMediaPlayer::EFeatureFlag::UsePlaybackTimingV2))
 	{
 		return FMediaTimeStamp(); // no media opened / old player timing mode
 	}
 
-	if (HaveVideoPlayback())
+	// Check if there are video samples present or presence is unknown.
+	// Only when we know for sure that there are none because the existing video stream has ended do we set this to false.
+	bool bHaveVideoSamples = VideoSampleAvailability != 0;
+
+	if (HaveVideoPlayback() && bHaveVideoSamples)
 	{
 		/*
 			Returning the precise time of the sample returned during TickFetch()
@@ -382,19 +473,34 @@ FMediaTimeStamp FMediaPlayerFacade::GetTimeStamp() const
 
 FText FMediaPlayerFacade::GetTrackDisplayName(EMediaTrackType TrackType, int32 TrackIndex) const
 {
-	return Player.IsValid() ? Player->GetTracks().GetTrackDisplayName((EMediaTrackType)TrackType, TrackIndex) : FText::GetEmpty();
+	TSharedPtr<IMediaPlayer, ESPMode::ThreadSafe> CurrentPlayer(Player);
+	if (!CurrentPlayer.IsValid())
+	{
+		return FText::GetEmpty();
+	}
+	return CurrentPlayer->GetTracks().GetTrackDisplayName((EMediaTrackType)TrackType, TrackIndex);
 }
 
 
 int32 FMediaPlayerFacade::GetTrackFormat(EMediaTrackType TrackType, int32 TrackIndex) const
 {
-	return Player.IsValid() ? Player->GetTracks().GetTrackFormat((EMediaTrackType)TrackType, TrackIndex) : INDEX_NONE;
+	TSharedPtr<IMediaPlayer, ESPMode::ThreadSafe> CurrentPlayer(Player);
+	if (!CurrentPlayer.IsValid())
+	{
+		return INDEX_NONE;
+	}
+	return CurrentPlayer->GetTracks().GetTrackFormat((EMediaTrackType)TrackType, TrackIndex);
 }
 
 
 FString FMediaPlayerFacade::GetTrackLanguage(EMediaTrackType TrackType, int32 TrackIndex) const
 {
-	return Player.IsValid() ? Player->GetTracks().GetTrackLanguage((EMediaTrackType)TrackType, TrackIndex) : FString();
+	TSharedPtr<IMediaPlayer, ESPMode::ThreadSafe> CurrentPlayer(Player);
+	if (!CurrentPlayer.IsValid())
+	{
+		return FString();
+	}
+	return CurrentPlayer->GetTracks().GetTrackLanguage((EMediaTrackType)TrackType, TrackIndex);
 }
 
 
@@ -435,72 +541,124 @@ FString FMediaPlayerFacade::GetVideoTrackType(int32 TrackIndex, int32 FormatInde
 
 bool FMediaPlayerFacade::GetViewField(float& OutHorizontal, float& OutVertical) const
 {
-	return Player.IsValid() ? Player->GetView().GetViewField(OutHorizontal, OutVertical) : false;
+	TSharedPtr<IMediaPlayer, ESPMode::ThreadSafe> CurrentPlayer(Player);
+	if (!CurrentPlayer.IsValid())
+	{
+		return false;
+	}
+	return CurrentPlayer->GetView().GetViewField(OutHorizontal, OutVertical);
 }
 
 
 bool FMediaPlayerFacade::GetViewOrientation(FQuat& OutOrientation) const
 {
-	return Player.IsValid() ? Player->GetView().GetViewOrientation(OutOrientation) : false;
+	TSharedPtr<IMediaPlayer, ESPMode::ThreadSafe> CurrentPlayer(Player);
+	if (!CurrentPlayer.IsValid())
+	{
+		return false;
+	}
+	return CurrentPlayer->GetView().GetViewOrientation(OutOrientation);
 }
 
 
 bool FMediaPlayerFacade::HasError() const
 {
-	return Player.IsValid() && (Player->GetControls().GetState() == EMediaState::Error);
+	TSharedPtr<IMediaPlayer, ESPMode::ThreadSafe> CurrentPlayer(Player);
+	if (!CurrentPlayer.IsValid())
+	{
+		return false;
+	}
+	return (CurrentPlayer->GetControls().GetState() == EMediaState::Error);
 }
 
 
 bool FMediaPlayerFacade::IsBuffering() const
 {
-	return Player.IsValid() && EnumHasAnyFlags(Player->GetControls().GetStatus(), EMediaStatus::Buffering);
+	TSharedPtr<IMediaPlayer, ESPMode::ThreadSafe> CurrentPlayer(Player);
+	if (!CurrentPlayer.IsValid())
+	{
+		return false;
+	}
+	return EnumHasAnyFlags(CurrentPlayer->GetControls().GetStatus(), EMediaStatus::Buffering);
 }
 
 
 bool FMediaPlayerFacade::IsConnecting() const
 {
-	return Player.IsValid() && EnumHasAnyFlags(Player->GetControls().GetStatus(), EMediaStatus::Connecting);
+	TSharedPtr<IMediaPlayer, ESPMode::ThreadSafe> CurrentPlayer(Player);
+	if (!CurrentPlayer.IsValid())
+	{
+		return false;
+	}
+	return EnumHasAnyFlags(CurrentPlayer->GetControls().GetStatus(), EMediaStatus::Connecting);
 }
 
 
 bool FMediaPlayerFacade::IsLooping() const
 {
-	return Player.IsValid() && Player->GetControls().IsLooping();
+	TSharedPtr<IMediaPlayer, ESPMode::ThreadSafe> CurrentPlayer(Player);
+	if (!CurrentPlayer.IsValid())
+	{
+		return false;
+	}
+	return CurrentPlayer->GetControls().IsLooping();
 }
 
 
 bool FMediaPlayerFacade::IsPaused() const
 {
-	return Player.IsValid() && (Player->GetControls().GetState() == EMediaState::Paused);
+	TSharedPtr<IMediaPlayer, ESPMode::ThreadSafe> CurrentPlayer(Player);
+	if (!CurrentPlayer.IsValid())
+	{
+		return false;
+	}
+	return (CurrentPlayer->GetControls().GetState() == EMediaState::Paused);
 }
 
 
 bool FMediaPlayerFacade::IsPlaying() const
 {
-	return Player.IsValid() && (Player->GetControls().GetState() == EMediaState::Playing);
+	TSharedPtr<IMediaPlayer, ESPMode::ThreadSafe> CurrentPlayer(Player);
+	if (!CurrentPlayer.IsValid())
+	{
+		return false;
+	}
+	return (CurrentPlayer->GetControls().GetState() == EMediaState::Playing);
 }
 
 
 bool FMediaPlayerFacade::IsPreparing() const
 {
-	return Player.IsValid() && (Player->GetControls().GetState() == EMediaState::Preparing);
+	TSharedPtr<IMediaPlayer, ESPMode::ThreadSafe> CurrentPlayer(Player);
+	if (!CurrentPlayer.IsValid())
+	{
+		return false;
+	}
+	return (CurrentPlayer->GetControls().GetState() == EMediaState::Preparing);
 }
 
 bool FMediaPlayerFacade::IsClosed() const
 {
-	return Player.IsValid() && (Player->GetControls().GetState() == EMediaState::Closed);
+	TSharedPtr<IMediaPlayer, ESPMode::ThreadSafe> CurrentPlayer(Player);
+	if (!CurrentPlayer.IsValid())
+	{
+		return false;
+	}
+	return (CurrentPlayer->GetControls().GetState() == EMediaState::Closed);
 }
 
 bool FMediaPlayerFacade::IsReady() const
 {
-	if (!Player.IsValid())
+	TSharedPtr<IMediaPlayer, ESPMode::ThreadSafe> CurrentPlayer(Player);
+	if (!CurrentPlayer.IsValid())
 	{
 		return false;
 	}
 
-	return ((Player->GetControls().GetState() != EMediaState::Closed) &&
-			(Player->GetControls().GetState() != EMediaState::Error) &&
-			(Player->GetControls().GetState() != EMediaState::Preparing));
+	EMediaState State = CurrentPlayer->GetControls().GetState();
+	return ((State != EMediaState::Closed) &&
+			(State != EMediaState::Error) &&
+			(State != EMediaState::Preparing));
 }
 
 
@@ -563,6 +721,7 @@ bool FMediaPlayerFacade::Open(const FString& Url, const IMediaOptions* Options,
 		return false;
 	}
 
+	FScopeLock Lock(&LastTimeValuesCS);
 	LastVideoSampleProcessedTime.Invalidate();
 	LastAudioSampleProcessedTime.Invalidate();
 	CurrentFrameAudioTimeStamp.Invalidate();
@@ -621,7 +780,7 @@ bool FMediaPlayerFacade::SelectTrack(EMediaTrackType TrackType, int32 TrackIndex
 	{
 		return false;
 	}
-	
+
 	FlushSinks();
 
 	return true;
@@ -825,7 +984,7 @@ bool FMediaPlayerFacade::GetAudioTrackFormat(int32 TrackIndex, int32 FormatIndex
 TSharedPtr<IMediaPlayer, ESPMode::ThreadSafe> FMediaPlayerFacade::GetPlayerForUrl(const FString& Url, const IMediaOptions* Options)
 {
 	FName PlayerName;
-	
+
 	if (DesiredPlayerName != NAME_None)
 	{
 		PlayerName = DesiredPlayerName;
@@ -918,7 +1077,7 @@ TSharedPtr<IMediaPlayer, ESPMode::ThreadSafe> FMediaPlayerFacade::GetPlayerForUr
 			else
 			{
 				UE_LOG(LogMediaUtils, Log, TEXT("| %s (only available on %s, but not on %s)"), *Factory->GetPlayerName().ToString(), *FString::Join(Factory->GetSupportedPlatforms(), TEXT(", ")), *RunningPlatformName);
-			}	
+			}
 		}
 	}
 	else
@@ -1044,7 +1203,7 @@ void FMediaPlayerFacade::TickInput(FTimespan DeltaTime, FTimespan Timecode)
 	{
 		Player->TickInput(DeltaTime, Timecode);
 
-		// Update flag reflecting presence of audio ion the current stream
+		// Update flag reflecting presence of audio in the current stream
 		// (doing it just once per gameloop is enough
 		bHaveActiveAudio = HaveAudioPlayback();
 
@@ -1109,6 +1268,9 @@ void FMediaPlayerFacade::TickInput(FTimespan DeltaTime, FTimespan Timecode)
 				CurrentFrameAudioTimeStamp = LastAudioSampleProcessedTime.TimeStamp;
 			}
 		}
+
+		// Check if primary audio sink needs a change and make sure invalid sinks are purged at all times
+		PrimaryAudioSink = AudioSampleSinks.GetPrimaryAudioSink();
 	}
 }
 
@@ -1280,9 +1442,9 @@ void FMediaPlayerFacade::TickTickable()
 
 	// process samples in range
 	IMediaSamples& Samples = Player->GetSamples();
-	
+
 	ProcessAudioSamples(Samples, AudioTimeRange);
-	ProcessMetadataSamples(Samples, MetadataTimeRange);	
+	ProcessMetadataSamples(Samples, MetadataTimeRange);
 
 	SET_DWORD_STAT(STAT_MediaUtils_FacadeNumAudioSamples, Samples.NumAudio());
 }
@@ -1375,10 +1537,14 @@ void FMediaPlayerFacade::ProcessAudioSamples(IMediaSamples& Samples, TRange<FTim
 {
 	TSharedPtr<IMediaAudioSample, ESPMode::ThreadSafe> Sample;
 
-	if (AudioSampleSinks.Num() == 1)
+	if (Player->GetPlayerFeatureFlag(IMediaPlayer::EFeatureFlag::UsePlaybackTimingV2))
 	{
+		// For V2 we basically expect to get no timerange at all: totally open
+		// (we just have it around to be compatible / use older code that expects it)
+		check(TimeRange.GetLowerBound().IsOpen() && TimeRange.GetUpperBound().IsOpen());
+
 		//
-		// "Modern" 1-Audio-Sink-Only case
+		// "Modern" 1-Audio-Sink-Only case (aka: we only feed the primary sink)
 		//
 		if (TSharedPtr< FMediaAudioSampleSink, ESPMode::ThreadSafe> PinnedPrimaryAudioSink = PrimaryAudioSink.Pin())
 		{
@@ -1392,11 +1558,34 @@ void FMediaPlayerFacade::ProcessAudioSamples(IMediaSamples& Samples, TRange<FTim
 					continue;
 				}
 
+				{
+				FScopeLock Lock(&LastTimeValuesCS);
 				LastAudioSampleProcessedTime.TimeStamp = FMediaTimeStamp(Sample->GetTime());
 				LastAudioSampleProcessedTime.SampledAtTime = FPlatformTime::Seconds();
+				}
 
-				// We assume we are the only ones adding samples: hence - after the above check - this MUST succeed!
-				verify(AudioSampleSinks.Enqueue(Sample.ToSharedRef(), FMediaPlayerQueueDepths::MaxAudioSinkDepth));
+				PinnedPrimaryAudioSink->Enqueue(Sample.ToSharedRef());
+			}
+		}
+		else
+		{
+			// Do we have video playback?
+			if (HaveVideoPlayback())
+			{
+				// We got video and audio, but no audio sink - throw away anything up to video playback time...
+				// (rough estimate, as this is off-gamethread; but better than throwing things out with no throttling at all)
+				{
+				FScopeLock Lock(&LastTimeValuesCS);
+				TimeRange.SetUpperBound(TRangeBound<FTimespan>(LastVideoSampleProcessedTime.TimeStamp.Time));
+				}
+				while (Samples.FetchAudio(TimeRange, Sample))
+					;
+			}
+			else
+			{
+				// No Video and no primary audio sink: we throw all away (sub-optimal as it will keep audio decoding busy; but this should be an edge case)
+				while (Samples.FetchAudio(TimeRange, Sample))
+					;
 			}
 		}
 	}
@@ -1421,6 +1610,7 @@ void FMediaPlayerFacade::ProcessAudioSamples(IMediaSamples& Samples, TRange<FTim
 			}
 			else
 			{
+				FScopeLock Lock(&LastTimeValuesCS);
 				LastAudioSampleProcessedTime.TimeStamp = FMediaTimeStamp(Sample->GetTime());
 				LastAudioSampleProcessedTime.SampledAtTime = FPlatformTime::Seconds();
 			}
@@ -1668,12 +1858,38 @@ void FMediaPlayerFacade::ReceiveMediaEvent(EMediaEvent Event)
 
 				bool bReverse = (Rate < 0.0f);
 				uint32 NumPurged = CurrentPlayer->GetSamples().PurgeOutdatedVideoSamples(TimeRange.GetLowerBoundValue() + (bReverse ? kOutdatedVideoSamplesTollerance : -kOutdatedVideoSamplesTollerance), bReverse);
-				
+
 				SET_DWORD_STAT(STAT_MediaUtils_FacadeNumPurgedVideoSamples, NumPurged);
 				INC_DWORD_STAT_BY(STAT_MediaUtils_FacadeTotalPurgedVideoSamples, NumPurged);
 
 				break;
 			}
+
+			case	EMediaEvent::Internal_VideoSamplesAvailable:
+			{
+				UE_LOG(LogMediaUtils, VeryVerbose, TEXT("PlayerFacade %p: Video samples ARE available"), this);
+				VideoSampleAvailability = 1;
+				break;
+			}
+			case	EMediaEvent::Internal_VideoSamplesUnavailable:
+			{
+				UE_LOG(LogMediaUtils, VeryVerbose, TEXT("PlayerFacade %p: Video samples are NOT available"), this);
+				VideoSampleAvailability = 0;
+				break;
+			}
+			case	EMediaEvent::Internal_AudioSamplesAvailable:
+			{
+				UE_LOG(LogMediaUtils, VeryVerbose, TEXT("PlayerFacade %p: Audio samples ARE available"), this);
+				AudioSampleAvailability = 1;
+				break;
+			}
+			case	EMediaEvent::Internal_AudioSamplesUnavailable:
+			{
+				UE_LOG(LogMediaUtils, VeryVerbose, TEXT("PlayerFacade %p: Audio samples are NOT available"), this);
+				AudioSampleAvailability = 0;
+				break;
+			}
+
 			default:
 				break;
 		}
diff --git a/Engine/Source/Runtime/MediaUtils/Public/MediaPlayerFacade.h b/Engine/Source/Runtime/MediaUtils/Public/MediaPlayerFacade.h
--- a/Engine/Source/Runtime/MediaUtils/Public/MediaPlayerFacade.h
+++ b/Engine/Source/Runtime/MediaUtils/Public/MediaPlayerFacade.h
@@ -817,6 +817,10 @@ private:
 	/** Flag indicating that we have an active audio setup */
 	bool bHaveActiveAudio;
 
+	/** Flag indicating the current availability of media samples. **/
+	int32 VideoSampleAvailability;
+	int32 AudioSampleAvailability;
+
 	/** An event delegate that is invoked when a media event occurred. */
 	FOnMediaEvent MediaEvent;
 
diff --git a/Engine/Source/Runtime/MediaUtils/Public/MediaSampleSinks.h b/Engine/Source/Runtime/MediaUtils/Public/MediaSampleSinks.h
--- a/Engine/Source/Runtime/MediaUtils/Public/MediaSampleSinks.h
+++ b/Engine/Source/Runtime/MediaUtils/Public/MediaSampleSinks.h
@@ -114,6 +114,20 @@ public:
 		Sinks.Remove(SampleSink);
 	}
 
+	/**
+	 * Remove any invalid sinks
+	 */
+	void Cleanup()
+	{
+		for (int32 SinkIndex = Sinks.Num() - 1; SinkIndex >= 0; --SinkIndex)
+		{
+			if (!Sinks[SinkIndex].IsValid())
+			{
+				Sinks.RemoveAtSwap(SinkIndex);
+			}
+		}
+	}
+
 protected:
 
 	/** The collection of registered sinks. */
@@ -125,8 +139,12 @@ typedef TMediaSampleSinks<IMediaTextureSample, FMediaTextureSampleSink> FMediaVi
 class FMediaAudioSampleSinks : public TMediaSampleSinks<IMediaAudioSample, FMediaAudioSampleSink>
 {
 public:
+	/**
+	 * Get primary audio sink and cleanup any invalid sinks
+	 */
 	TSharedPtr<FMediaAudioSampleSink, ESPMode::ThreadSafe> GetPrimaryAudioSink()
 	{
+		Cleanup();
 		if (Sinks.Num() == 0)
 		{
 			return TSharedPtr<FMediaAudioSampleSink, ESPMode::ThreadSafe>();
