Electra: when rebuffering in a Live stream we now implicitly jump to the Live edge to prevent users from falling too far behind

#rnx

#localization none
#tests tested in PIE on Windows

diff --git a/Engine/Source/Runtime/Media/Public/IMediaEventSink.h b/Engine/Source/Runtime/Media/Public/IMediaEventSink.h
--- a/Engine/Source/Runtime/Media/Public/IMediaEventSink.h
+++ b/Engine/Source/Runtime/Media/Public/IMediaEventSink.h
@@ -46,6 +46,13 @@ enum class EMediaEvent
 	/** Player requests purge of outdated video samples from queues. */
 	Internal_PurgeVideoSamplesHint = Internal_Start,
 
+	/** Player requests an internal reset of all queues and timing. **/
+	Internal_ResetForDiscontinuity,
+
+	/** Player requests render clock to start and stop in preparation for moving playback to a new position. **/
+	Internal_RenderClockStart,
+	Internal_RenderClockStop,
+
 	/** Player notifies about availability or unavailablity of media samples. **/
 	Internal_VideoSamplesAvailable,
 	Internal_VideoSamplesUnavailable,
diff --git a/Engine/Source/Runtime/MediaUtils/Private/MediaPlayerFacade.cpp b/Engine/Source/Runtime/MediaUtils/Private/MediaPlayerFacade.cpp
--- a/Engine/Source/Runtime/MediaUtils/Private/MediaPlayerFacade.cpp
+++ b/Engine/Source/Runtime/MediaUtils/Private/MediaPlayerFacade.cpp
@@ -242,6 +242,7 @@ void FMediaPlayerFacade::Close()
 	bHaveActiveAudio = false;
 	VideoSampleAvailability = -1;
 	AudioSampleAvailability = -1;
+	bIsSinkFlushPending = false;
 
 	FlushSinks();
 }
@@ -1224,6 +1225,11 @@ void FMediaPlayerFacade::TickInput(FTimespan DeltaTime, FTimespan Timecode)
 			//
 			// New timing control (handled before any engine world, object etc. updates; so "all frame" (almost) see the state produced here)
 			//
+			if (bIsSinkFlushPending)
+			{
+				bIsSinkFlushPending = false;
+				FlushSinks();
+			}
 
 			// Do we have a current timestamp estimation?
 			if (!bHaveActiveAudio && !NextEstVideoTimeAtFrameStart.IsValid())
@@ -1865,6 +1871,23 @@ void FMediaPlayerFacade::ReceiveMediaEvent(EMediaEvent Event)
 				break;
 			}
 
+			case	EMediaEvent::Internal_ResetForDiscontinuity:
+			{
+				UE_LOG(LogMediaUtils, VeryVerbose, TEXT("PlayerFacade %p: Reset for discontinuity"), this);
+				bIsSinkFlushPending = true;
+				break;
+			}
+			case	EMediaEvent::Internal_RenderClockStart:
+			{
+				UE_LOG(LogMediaUtils, VeryVerbose, TEXT("PlayerFacade %p: Render clock shall start"), this);
+				break;
+			}
+			case	EMediaEvent::Internal_RenderClockStop:
+			{
+				UE_LOG(LogMediaUtils, VeryVerbose, TEXT("PlayerFacade %p: Render clock shall stop"), this);
+				break;
+			}
+
 			case	EMediaEvent::Internal_VideoSamplesAvailable:
 			{
 				UE_LOG(LogMediaUtils, VeryVerbose, TEXT("PlayerFacade %p: Video samples ARE available"), this);
diff --git a/Engine/Source/Runtime/MediaUtils/Public/MediaPlayerFacade.h b/Engine/Source/Runtime/MediaUtils/Public/MediaPlayerFacade.h
--- a/Engine/Source/Runtime/MediaUtils/Public/MediaPlayerFacade.h
+++ b/Engine/Source/Runtime/MediaUtils/Public/MediaPlayerFacade.h
@@ -854,6 +854,9 @@ private:
 	/** Estimation for next frame's video timestamp (used when no audio present or active in stream) */
 	FMediaTimeStamp NextEstVideoTimeAtFrameStart;
 
+	/** Set if sinks are to be flushed at the request of the player. **/
+	TAtomic<bool>	bIsSinkFlushPending;
+
 	/** Mediamodule we are working in */
 	IMediaModule* MediaModule;
 };
