Introducing a new "lifecycle manager delegate" to Mediaframework to allow tracking and influencing of player creation / destruction framework-wide

diff --git a/Engine/Source/Runtime/Media/Private/MediaModule.cpp b/Engine/Source/Runtime/Media/Private/MediaModule.cpp
--- a/Engine/Source/Runtime/Media/Private/MediaModule.cpp
+++ b/Engine/Source/Runtime/Media/Private/MediaModule.cpp
@@ -139,6 +139,24 @@ public:
 		PlayerFactories.AddUnique(&Factory);
 	}
 
+	virtual void SetPlayerLifecycleManagerDelegate(IMediaPlayerLifecycleManagerDelegate* Delegate) override
+	{
+		PlayerLifecycleManagerDelegate = Delegate;
+	}
+
+	virtual IMediaPlayerLifecycleManagerDelegate* GetPlayerLifecycleManagerDelegate() override
+	{
+		return PlayerLifecycleManagerDelegate;
+	}
+
+	virtual uint64 CreateMediaPlayerInstanceID() override
+	{
+		uint64 InstanceID;
+		while ((InstanceID = NextMediaPlayerInstanceID++) == ~0)
+			;
+		return InstanceID;
+	}
+
 	virtual void SetTimeSource(const TSharedPtr<IMediaTimeSource, ESPMode::ThreadSafe>& NewTimeSource) override
 	{
 		TimeSource = NewTimeSource;
@@ -256,6 +274,12 @@ private:
 	/** The registered video player factories. */
 	TArray<IMediaPlayerFactory*> PlayerFactories;
 
+	/** Player lifecycle manager delegate */
+	IMediaPlayerLifecycleManagerDelegate* PlayerLifecycleManagerDelegate = nullptr;
+
+	/** MediaPlayer instance ID to hand out next */
+	uint64 NextMediaPlayerInstanceID = 0;
+
 	/** High-frequency ticker runnable. */
 	FMediaTicker Ticker;
 
@@ -266,7 +290,7 @@ private:
 	FSimpleMulticastDelegate OnTickPreEngineCompleted;
 
 	/** Whether media objects should lock to the media clock's time code. */
-	bool TimecodeLocked;
+	bool TimecodeLocked = false;
 
 	/** The media clock's time source. */
 	TSharedPtr<IMediaTimeSource, ESPMode::ThreadSafe> TimeSource;
diff --git a/Engine/Source/Runtime/Media/Public/IMediaModule.h b/Engine/Source/Runtime/Media/Public/IMediaModule.h
--- a/Engine/Source/Runtime/Media/Public/IMediaModule.h
+++ b/Engine/Source/Runtime/Media/Public/IMediaModule.h
@@ -15,7 +15,7 @@ class IMediaPlayerFactory;
 class IMediaTicker;
 class IMediaTimeSource;
 class IMediaInfo;
-
+class IMediaPlayerLifecycleManagerDelegate;
 
 /**
  * Interface for the Media module.
@@ -125,6 +125,22 @@ public:
 	 */
 	virtual void UnregisterPlayerFactory(IMediaPlayerFactory& Factory) = 0;
 
+	/**
+	* Set player lifetime delegate
+	* @param Delegate Delegate instance to set for use. Use nullptr to reset delegate.
+	*/
+	virtual void SetPlayerLifecycleManagerDelegate(IMediaPlayerLifecycleManagerDelegate* Delegate) = 0;
+
+	/**
+	* Get player lifetime delegate
+	*/
+	virtual IMediaPlayerLifecycleManagerDelegate* GetPlayerLifecycleManagerDelegate() = 0;
+
+	/**
+	* Get new media player instance ID
+	*/
+	virtual uint64 CreateMediaPlayerInstanceID() = 0;
+
 public:
 
 	//~ Time management
diff --git a/Engine/Source/Runtime/Media/Public/IMediaPlayer.h b/Engine/Source/Runtime/Media/Public/IMediaPlayer.h
--- a/Engine/Source/Runtime/Media/Public/IMediaPlayer.h
+++ b/Engine/Source/Runtime/Media/Public/IMediaPlayer.h
@@ -318,6 +318,20 @@ public:
 		return false;
 	}
 
+	class IAsyncDestructNotification
+	{
+	public:
+		virtual ~IAsyncDestructNotification() {}
+		virtual void Signal() = 0;
+	};
+	typedef TSharedRef<IAsyncDestructNotification, ESPMode::ThreadSafe> IAsyncDestructNotificationRef;
+
+	virtual bool SetAsyncDestructionNotification(IAsyncDestructNotificationRef AsyncDestructNotification)
+	{
+		// Override in child class if needed.
+		return false;
+	}
+
 public:
 
 	/** Virtual destructor. */
diff --git a/Engine/Source/Runtime/MediaUtils/Private/MediaPlayerFacade.cpp b/Engine/Source/Runtime/MediaUtils/Private/MediaPlayerFacade.cpp
--- a/Engine/Source/Runtime/MediaUtils/Private/MediaPlayerFacade.cpp
+++ b/Engine/Source/Runtime/MediaUtils/Private/MediaPlayerFacade.cpp
@@ -28,6 +28,8 @@
 #include "MediaSampleQueueDepths.h"
 #include "MediaSampleQueue.h"
 
+#include "Async/Async.h"
+
 #define MEDIAPLAYERFACADE_DISABLE_BLOCKING 0
 #define MEDIAPLAYERFACADE_TRACE_SINKOVERFLOWS 0
 
@@ -107,10 +109,13 @@ FMediaPlayerFacade::~FMediaPlayerFacade()
 {
 	if (Player.IsValid())
 	{
-		FScopeLock Lock(&CriticalSection);
+		{
+			FScopeLock Lock(&CriticalSection);
+			Player->Close();
+		}
+		NotifyLifetimeManagerDelegate_PlayerClosed();
 
-		Player->Close();
-		Player.Reset();
+		DestroyPlayer();
 	}
 
 	delete Cache;
@@ -232,8 +237,11 @@ void FMediaPlayerFacade::Close()
 	TSharedPtr<IMediaPlayer, ESPMode::ThreadSafe> CurrentPlayer(Player);
 	if (CurrentPlayer.IsValid())
 	{
-		FScopeLock Lock(&CriticalSection);
-		CurrentPlayer->Close();
+		{
+			FScopeLock Lock(&CriticalSection);
+			CurrentPlayer->Close();
+		}
+		NotifyLifetimeManagerDelegate_PlayerClosed();
 	}
 
 	BlockOnTime = FTimespan::MinValue();
@@ -665,6 +673,214 @@ bool FMediaPlayerFacade::IsReady() const
 			(State != EMediaState::Preparing));
 }
 
+// ----------------------------------------------------------------------------------------------------------------------------------------------
+
+class FMediaPlayerLifecycleManagerDelegateOpenRequest : public IMediaPlayerLifecycleManagerDelegate::IOpenRequest
+{
+public:
+	FMediaPlayerLifecycleManagerDelegateOpenRequest(const FString& InUrl, const IMediaOptions* InOptions, const FMediaPlayerOptions* InPlayerOptions, IMediaPlayerFactory* InPlayerFactory, bool bInWillCreatePlayer)
+		: Url(InUrl), Options(InOptions), PlayerFactory(InPlayerFactory), bWillCreatePlayer(bInWillCreatePlayer)
+	{
+		if (InPlayerOptions)
+		{
+			PlayerOptions = *InPlayerOptions;
+		}
+	}
+
+	virtual const FString& GetUrl() const override
+	{
+		return Url;
+	}
+
+	virtual const IMediaOptions* GetOptions() const override
+	{
+		return Options;
+	}
+
+	virtual const FMediaPlayerOptions* GetPlayerOptions() const override
+	{
+		return PlayerOptions.IsSet() ? &PlayerOptions.GetValue() : nullptr;
+	}
+
+	virtual IMediaPlayerFactory* GetPlayerFactory() const override
+	{
+		return PlayerFactory;
+	}
+
+	virtual bool WillCreateNewPlayer() const
+	{
+		return bWillCreatePlayer;
+	}
+
+	FString Url;
+	const IMediaOptions* Options;
+	TOptional<FMediaPlayerOptions> PlayerOptions;
+	IMediaPlayerFactory* PlayerFactory;
+	bool bWillCreatePlayer;
+};
+
+class FMediaPlayerLifecycleManagerDelegateControl : public IMediaPlayerLifecycleManagerDelegate::IControl
+{
+public:
+	FMediaPlayerLifecycleManagerDelegateControl(TWeakPtr<FMediaPlayerFacade, ESPMode::ThreadSafe> InFacade) : Facade(InFacade), InstanceID(~0), SubmittedRequest(false) {}
+
+	virtual ~FMediaPlayerLifecycleManagerDelegateControl()
+	{
+		if (!SubmittedRequest)
+		{
+			if (TSharedPtr<FMediaPlayerFacade, ESPMode::ThreadSafe> PinnedFacade = Facade.Pin())
+			{
+				PinnedFacade->ReceiveMediaEvent(EMediaEvent::MediaOpenFailed);
+			}
+		}
+	}
+
+	virtual bool SubmitOpenRequest(IMediaPlayerLifecycleManagerDelegate::IOpenRequestRef&& OpenRequest) override
+	{
+		if (TSharedPtr<FMediaPlayerFacade, ESPMode::ThreadSafe> PinnedFacade = Facade.Pin())
+		{
+			const FMediaPlayerLifecycleManagerDelegateOpenRequest* OR = static_cast<const FMediaPlayerLifecycleManagerDelegateOpenRequest*>(OpenRequest.Get());
+			bool bOk = PinnedFacade->ContinueOpen(OR->Url, OR->Options, OR->PlayerOptions.IsSet() ? &OR->PlayerOptions.GetValue() : nullptr, OR->PlayerFactory, OR->bWillCreatePlayer);
+			if (bOk)
+			{
+				SubmittedRequest = true;
+			}
+			return bOk;
+		}
+		return false;
+	}
+
+	virtual TSharedPtr<FMediaPlayerFacade, ESPMode::ThreadSafe> GetFacade() const override
+	{
+		return Facade.Pin();
+	}
+
+	virtual uint64 GetMediaPlayerInstanceID() const override
+	{
+		return InstanceID;
+	}
+
+	void SetInstanceID(uint64 InInstanceID)
+	{
+		InstanceID = InInstanceID;
+	}
+
+private:
+	TWeakPtr<FMediaPlayerFacade, ESPMode::ThreadSafe> Facade;
+	uint64 InstanceID;
+	bool SubmittedRequest;
+};
+
+// ----------------------------------------------------------------------------------------------------------------------------------------------
+
+bool FMediaPlayerFacade::NotifyLifetimeManagerDelegate_PlayerOpen(const FString& Url, const IMediaOptions* Options, const FMediaPlayerOptions* PlayerOptions, IMediaPlayerFactory* PlayerFactory, bool bWillCreatePlayer)
+{
+	check(IsInGameThread());
+
+	if (IMediaPlayerLifecycleManagerDelegate* Delegate = MediaModule->GetPlayerLifecycleManagerDelegate())
+	{
+		LifecycleManagerDelegateControl = MakeShared<FMediaPlayerLifecycleManagerDelegateControl, ESPMode::ThreadSafe>(AsShared());
+		if (LifecycleManagerDelegateControl.IsValid())
+		{
+			IMediaPlayerLifecycleManagerDelegate::IOpenRequestRef OpenRequest(new FMediaPlayerLifecycleManagerDelegateOpenRequest(Url, Options, PlayerOptions, PlayerFactory, bWillCreatePlayer));
+			if (OpenRequest.IsValid())
+			{
+				return Delegate->OnMediaPlayerOpen(LifecycleManagerDelegateControl, OpenRequest);
+			}
+		}
+	}
+	return false;
+}
+
+bool FMediaPlayerFacade::NotifyLifetimeManagerDelegate_PlayerCreated()
+{
+	check(IsInGameThread());
+	check(Player.IsValid());
+
+	if (LifecycleManagerDelegateControl.IsValid())
+	{
+		static_cast<FMediaPlayerLifecycleManagerDelegateControl*>(LifecycleManagerDelegateControl.Get())->SetInstanceID(PlayerInstanceID);
+
+		if (IMediaPlayerLifecycleManagerDelegate* Delegate = MediaModule->GetPlayerLifecycleManagerDelegate())
+		{
+			Delegate->OnMediaPlayerCreated(LifecycleManagerDelegateControl);
+			return true;
+		}
+	}
+	return false;
+}
+
+bool FMediaPlayerFacade::NotifyLifetimeManagerDelegate_PlayerClosed()
+{
+	check(IsInGameThread());
+
+	if (LifecycleManagerDelegateControl.IsValid())
+	{
+		if (IMediaPlayerLifecycleManagerDelegate* Delegate = MediaModule->GetPlayerLifecycleManagerDelegate())
+		{
+			Delegate->OnMediaPlayerClosed(LifecycleManagerDelegateControl);
+			return true;
+		}
+	}
+	return false;
+}
+
+bool FMediaPlayerFacade::NotifyLifetimeManagerDelegate_PlayerDestroyed()
+{
+	check(IsInGameThread());
+
+	if (LifecycleManagerDelegateControl.IsValid())
+	{
+		if (IMediaPlayerLifecycleManagerDelegate* Delegate = MediaModule->GetPlayerLifecycleManagerDelegate())
+		{
+			Delegate->OnMediaPlayerDestroyed(LifecycleManagerDelegateControl);
+			return true;
+		}
+	}
+	return false;
+}
+
+// ----------------------------------------------------------------------------------------------------------------------------------------------
+
+void FMediaPlayerFacade::DestroyPlayer()
+{
+	FScopeLock Lock(&CriticalSection);
+
+	if (!Player.IsValid() || !LifecycleManagerDelegateControl.IsValid())
+	{
+		return;
+	}
+
+	class FAsyncDestructNotification : public IMediaPlayer::IAsyncDestructNotification
+	{
+	public:
+		FAsyncDestructNotification(IMediaModule* InMediaModule, IMediaPlayerLifecycleManagerDelegate::IControlRef InDelegateControl) : MediaModule(InMediaModule), DelegateControl(InDelegateControl) {}
+
+		virtual void Signal() override
+		{
+			TFunction<void()> NotifyTask = [TargetMediaModule{ MoveTemp(MediaModule) }, TargetDelegateControl { MoveTemp(DelegateControl) }]()
+			{
+				if (IMediaPlayerLifecycleManagerDelegate* Delegate = TargetMediaModule->GetPlayerLifecycleManagerDelegate())
+				{
+					Delegate->OnMediaPlayerDestroyed(TargetDelegateControl);
+				}
+			};
+			Async(EAsyncExecution::TaskGraphMainThread, NotifyTask);
+		};
+
+		IMediaModule* MediaModule;
+		IMediaPlayerLifecycleManagerDelegate::IControlRef DelegateControl;
+	};
+
+	bool bWillNotify = Player->SetAsyncDestructionNotification(TSharedRef<IMediaPlayer::IAsyncDestructNotification, ESPMode::ThreadSafe>(new FAsyncDestructNotification(MediaModule, LifecycleManagerDelegateControl)));
+	Player.Reset();
+	if (!bWillNotify)
+	{
+		// The player did not implement the async destruction notification, so we assume it did everything synchronously...
+		// (note: there is a chance something else is hanging on to the player's ref and we are triggering this notification too early, but as we expect all V2 players to implement this eventually, we accept this for now)
+		NotifyLifetimeManagerDelegate_PlayerDestroyed();
+	}
+}
 
 bool FMediaPlayerFacade::Open(const FString& Url, const IMediaOptions* Options, const FMediaPlayerOptions* PlayerOptions)
 {
@@ -686,13 +902,37 @@ bool FMediaPlayerFacade::Open(const FString& Url, const IMediaOptions* Options,
 
 	check(MediaModule);
 
-	// find & initialize new player
-	TSharedPtr<IMediaPlayer, ESPMode::ThreadSafe> NewPlayer = GetPlayerForUrl(Url, Options);
+	// find a player factory for the intended playback
+	IMediaPlayerFactory* PlayerFactory = GetPlayerFactoryForUrl(Url, Options);
+
+	IMediaPlayerFactory* OldFactory(Player.IsValid() ? MediaModule->GetPlayerFactory(Player->GetPlayerPluginGUID()) : nullptr);
+
+	bool bWillCreatePlayer = (PlayerFactory != OldFactory);
+
+	if (FMediaPlayerFacade::NotifyLifetimeManagerDelegate_PlayerOpen(Url, Options, PlayerOptions, PlayerFactory, bWillCreatePlayer))
+	{
+		// Assume all is well: the delegate will either (have) submit(ted) the request or not -- in any case we need to assume the best -> "true"
+		return true;
+	}
+
+	// We did not notify successfully or the delegate will not submit the request in its own. Do so here...
+	return ContinueOpen(Url, Options, PlayerOptions, PlayerFactory, bWillCreatePlayer);
+}
+
+bool FMediaPlayerFacade::ContinueOpen(const FString& Url, const IMediaOptions* Options, const FMediaPlayerOptions* PlayerOptions, IMediaPlayerFactory* PlayerFactory, bool bCreateNewPlayer)
+{
+	// Create or reuse player
+	TSharedPtr<IMediaPlayer, ESPMode::ThreadSafe> NewPlayer(bCreateNewPlayer ? PlayerFactory->CreatePlayer(*this) : Player);
+
+	// Continue initialization ---------------------------------------
 
 	if (NewPlayer != Player)
 	{
+		DestroyPlayer();
+
 		FScopeLock Lock(&CriticalSection);
 		Player = NewPlayer;
+		PlayerInstanceID = MediaModule->CreateMediaPlayerInstanceID();
 	}
 
 	if (!Player.IsValid())
@@ -730,6 +970,11 @@ bool FMediaPlayerFacade::Open(const FString& Url, const IMediaOptions* Options,
 	LastAudioSampleProcessedTime.Invalidate();
 	CurrentFrameAudioTimeStamp.Invalidate();
 
+	if (Player.IsValid() && bCreateNewPlayer)
+	{
+		NotifyLifetimeManagerDelegate_PlayerCreated();
+	}
+
 	return true;
 }
 
@@ -985,7 +1230,7 @@ bool FMediaPlayerFacade::GetAudioTrackFormat(int32 TrackIndex, int32 FormatIndex
 }
 
 
-TSharedPtr<IMediaPlayer, ESPMode::ThreadSafe> FMediaPlayerFacade::GetPlayerForUrl(const FString& Url, const IMediaOptions* Options)
+IMediaPlayerFactory* FMediaPlayerFacade::GetPlayerFactoryForUrl(const FString& Url, const IMediaOptions* Options) const
 {
 	FName PlayerName;
 
@@ -1008,13 +1253,21 @@ TSharedPtr<IMediaPlayer, ESPMode::ThreadSafe> FMediaPlayerFacade::GetPlayerForUr
 		return nullptr;
 	}
 
-	// reuse existing player if requested
-	if (Player.IsValid() && (PlayerName == MediaModule->GetPlayerFactory(Player->GetPlayerPluginGUID())->GetPlayerName()))
+	//
+	// Reuse existing player if explicitly requested name matches
+	//
+	if (Player.IsValid())
 	{
-		return Player;
+		IMediaPlayerFactory* CurrentFactory = MediaModule->GetPlayerFactory(Player->GetPlayerPluginGUID());
+		if (PlayerName == CurrentFactory->GetPlayerName())
+		{
+			return CurrentFactory;
+		}
 	}
 
-	// try to create requested player
+	//
+	// Try to create explicitly requested player
+	//
 	if (PlayerName != NAME_None)
 	{
 		IMediaPlayerFactory* Factory = MediaModule->GetPlayerFactory(PlayerName);
@@ -1022,34 +1275,32 @@ TSharedPtr<IMediaPlayer, ESPMode::ThreadSafe> FMediaPlayerFacade::GetPlayerForUr
 		if (Factory == nullptr)
 		{
 			UE_LOG(LogMediaUtils, Error, TEXT("Could not find desired player %s for %s"), *PlayerName.ToString(), *Url);
-			return nullptr;
 		}
 
-		TSharedPtr<IMediaPlayer, ESPMode::ThreadSafe> NewPlayer = Factory->CreatePlayer(*this);
+		return Factory;
+	}
+
 
-		if (!NewPlayer.IsValid())
-		{
-			UE_LOG(LogMediaUtils, Error, TEXT("Failed to create desired player %s for %s"), *PlayerName.ToString(), *Url);
-			return nullptr;
-		}
 
-		return NewPlayer;
-	}
+	//
+	// Try to find a fitting player with no explicit name given
+	//
+
 
-	// try to reuse existing player
+	// Can any existing player play the URL?
 	if (Player.IsValid())
 	{
 		IMediaPlayerFactory* Factory = MediaModule->GetPlayerFactory(Player->GetPlayerPluginGUID());
 
 		if ((Factory != nullptr) && Factory->CanPlayUrl(Url, Options))
 		{
-			return Player;
+			// Yes...
+			return Factory;
 		}
 	}
 
+	// Try to auto-select new player...
 	const FString RunningPlatformName(FPlatformProperties::IniPlatformName());
-
-	// try to auto-select new player
 	const TArray<IMediaPlayerFactory*>& PlayerFactories = MediaModule->GetPlayerFactories();
 
 	for (IMediaPlayerFactory* Factory : PlayerFactories)
@@ -1059,15 +1310,12 @@ TSharedPtr<IMediaPlayer, ESPMode::ThreadSafe> FMediaPlayerFacade::GetPlayerForUr
 			continue;
 		}
 
-		TSharedPtr<IMediaPlayer, ESPMode::ThreadSafe> NewPlayer = Factory->CreatePlayer(*this);
-
-		if (NewPlayer.IsValid())
-		{
-			return NewPlayer;
-		}
+		return Factory;
 	}
 
-	// no suitable player found
+	//
+	// No suitable player found!
+	//
 	if (PlayerFactories.Num() > 0)
 	{
 		UE_LOG(LogMediaUtils, Error, TEXT("Cannot play %s, because none of the enabled media player plug-ins support it:"), *Url);
@@ -1159,8 +1407,7 @@ void FMediaPlayerFacade::ProcessEvent(EMediaEvent Event)
 			if(Player.IsValid() && Player->GetPlayerFeatureFlag(IMediaPlayer::EFeatureFlag::AllowShutdownOnClose))
 			{
 				bDidRecentPlayerHaveError = HasError();
-				FScopeLock Lock(&CriticalSection);
-				Player.Reset();
+				DestroyPlayer();
 			}
 
 			// Stop issuing audio thread ticks until we open the player again
diff --git a/Engine/Source/Runtime/MediaUtils/Public/IMediaPlayerLifecycleManager.h b/Engine/Source/Runtime/MediaUtils/Public/IMediaPlayerLifecycleManager.h
new file mode 100644
--- /dev/null
+++ b/Engine/Source/Runtime/MediaUtils/Public/IMediaPlayerLifecycleManager.h
@@ -0,0 +1,75 @@
+// Copyright Epic Games, Inc. All Rights Reserved.
+
+#pragma once
+
+#include "CoreTypes.h"
+#include "Containers/UnrealString.h"
+#include "Templates/SharedPointer.h"
+
+class FMediaPlayerFacade;
+class IMediaOptions;
+class IMediaPlayerFactory;
+
+struct FMediaPlayerOptions;
+
+/*
+	A single IMediaPlayerLifecycleManagerDelegate can be registered with the MediaModule to monitor and control
+	player lifecycles throughout the media framework.
+	One possible use case motivating the implementation is to control player creation as system resources are being
+	monitored. As this task is highly dependent on knowledge the media framework itself often does not have, only
+	an application-supplied delegate can be expected to make reasonably good decisions.
+
+	The delegate will receive notifications about major lifecycle events of all players handled by the framework.
+
+	To identify the instances the implementation should not rely on pointer comparisons (the pointers in question
+	may be used in reallocations or may be no longer valid for some notifications), but should rather use the
+	supplied 64-bit ID, which is uniquely (enough) generated for each created instance.
+
+	Notes:
+
+	- All notification callbacks are issues on the game thread
+	- The PlayerOpen callback will be issued before any player instance has been created, hence the ID is invalid (but a facade may exist)
+	-- If the callback returns "true" it must take care to call SubmitOpenRequest() with the passed along open request as soon as it deems it "ok" for the player to serve the request to be created
+	-- (if not "submitted" the open will not create an actual player; if "false" is returned the default creation method will continue as if a submit was done)
+	- GetFacade() may return a null pointer in case the facade instance was destroyed since the event was triggered and hence should always be checked before use
+	- The PlayerDestroyed callback will be triggered after the player is already destroyed. The facade may or may not still exist. The ID is still valid so this event can be properly tracked
+
+*/
+//! Interface to receive global player lifetime events from media framework
+class IMediaPlayerLifecycleManagerDelegate
+{
+public:
+	//! Request to create and open a player
+	class IOpenRequest
+	{
+	public:
+		virtual ~IOpenRequest() {}
+
+		virtual const FString& GetUrl() const = 0;
+		virtual const IMediaOptions* GetOptions() const = 0;
+		virtual const FMediaPlayerOptions* GetPlayerOptions() const = 0;
+		virtual IMediaPlayerFactory* GetPlayerFactory() const = 0;
+		virtual bool WillCreateNewPlayer() const = 0;
+	};
+	typedef TSharedPtr<IOpenRequest, ESPMode::ThreadSafe> IOpenRequestRef;
+
+	//! Control interface for lifecycle delegate
+	class IControl
+	{
+	public:
+		virtual ~IControl() {}
+
+		virtual bool SubmitOpenRequest(IOpenRequestRef && OpenRequest) = 0;
+
+		virtual TSharedPtr<FMediaPlayerFacade, ESPMode::ThreadSafe> GetFacade() const = 0;
+		virtual uint64 GetMediaPlayerInstanceID() const = 0;
+	};
+	typedef TSharedPtr<IControl, ESPMode::ThreadSafe> IControlRef;
+
+	virtual ~IMediaPlayerLifecycleManagerDelegate() {}
+
+	virtual bool OnMediaPlayerOpen(IControlRef Control, IOpenRequestRef OpenRequest) = 0;
+	virtual void OnMediaPlayerCreated(IControlRef Control) = 0;
+	virtual void OnMediaPlayerClosed(IControlRef Control) = 0;
+	virtual void OnMediaPlayerDestroyed(IControlRef Control) = 0;
+};
diff --git a/Engine/Source/Runtime/MediaUtils/Public/MediaPlayerFacade.h b/Engine/Source/Runtime/MediaUtils/Public/MediaPlayerFacade.h
--- a/Engine/Source/Runtime/MediaUtils/Public/MediaPlayerFacade.h
+++ b/Engine/Source/Runtime/MediaUtils/Public/MediaPlayerFacade.h
@@ -10,6 +10,7 @@
 #include "IMediaEventSink.h"
 #include "IMediaTickable.h"
 #include "IMediaTimeSource.h"
+#include "IMediaPlayerLifecycleManager.h"
 #include "MediaPlayerOptions.h"
 #include "Math/Quat.h"
 #include "Math/Range.h"
@@ -26,6 +27,7 @@ class IMediaOptions;
 class IMediaPlayer;
 class IMediaSamples;
 class IMediaModule;
+class IMediaPlayerFactory;
 
 enum class EMediaEvent;
 enum class EMediaCacheState;
@@ -34,6 +36,7 @@ enum class EMediaTrackType;
 
 struct FMediaAudioTrackFormat;
 struct FMediaVideoTrackFormat;
+struct FMediaPlayerOptions;
 
 
 /**
@@ -725,7 +728,7 @@ protected:
 	 * @param Options The media options for the URL.
 	 * @return The player if found, or nullptr otherwise.
 	 */
-	TSharedPtr<IMediaPlayer, ESPMode::ThreadSafe> GetPlayerForUrl(const FString& Url, const IMediaOptions* Options);
+	IMediaPlayerFactory *GetPlayerFactoryForUrl(const FString& Url, const IMediaOptions* Options) const;
 
 	/**
 	 * Get details about the specified audio track format.
@@ -774,6 +777,13 @@ protected:
 	void ReceiveMediaEvent(EMediaEvent Event) override;
 
 private:
+	friend class FMediaPlayerLifecycleManagerDelegateControl;
+
+	// Internal
+	bool NotifyLifetimeManagerDelegate_PlayerOpen(const FString& InUrl, const IMediaOptions* Options, const FMediaPlayerOptions* InPlayerOptions, IMediaPlayerFactory* PlayerFactory, bool bWillCreatePlayer);
+	bool NotifyLifetimeManagerDelegate_PlayerCreated();
+	bool NotifyLifetimeManagerDelegate_PlayerClosed();
+	bool NotifyLifetimeManagerDelegate_PlayerDestroyed();
 
 	void ProcessVideoSamples(IMediaSamples& Samples, const TRange<FMediaTimeStamp> & TimeRange);
 	void ProcessCaptionSamples(IMediaSamples& Samples, TRange<FMediaTimeStamp> TimeRange);
@@ -781,6 +791,10 @@ private:
 
 	bool GetCurrentPlaybackTimeRange(TRange<FMediaTimeStamp> & TimeRange, float Rate, FTimespan DeltaTime, bool bDoNotUseFrameStartReference) const;
 
+	void DestroyPlayer();
+
+	bool ContinueOpen(const FString& Url, const IMediaOptions* Options, const FMediaPlayerOptions* PlayerOptions, IMediaPlayerFactory* PlayerFactory, bool bCreateNewPlayer);
+
 	/** Audio sample sinks. */
 	TWeakPtr<FMediaAudioSampleSink, ESPMode::ThreadSafe> PrimaryAudioSink;
 	FMediaAudioSampleSinks AudioSampleSinks;
@@ -830,6 +844,9 @@ private:
 	/** The low-level player used to play the media source. */
 	TSharedPtr<IMediaPlayer, ESPMode::ThreadSafe> Player;
 
+	/** Low level player instance ID */
+	uint64 PlayerInstanceID = ~0;
+
 	/** Media player Guid */
 	FGuid PlayerGuid;
 
@@ -862,4 +879,7 @@ private:
 
 	/** Mediamodule we are working in */
 	IMediaModule* MediaModule;
+
+	/** Control interface for lifecycle manager delegate */
+	IMediaPlayerLifecycleManagerDelegate::IControlRef LifecycleManagerDelegateControl;
 };
