Making the instance ID of media players available already in the pre-createion open callback to the lifetime manager & correcting design of lifecycle manager callbacks a bit
(contains  adjustments of F5Palyer/Electra, which used this API already)

diff --git a/Engine/Source/Runtime/Media/Public/IMediaPlayer.h b/Engine/Source/Runtime/Media/Public/IMediaPlayer.h
--- a/Engine/Source/Runtime/Media/Public/IMediaPlayer.h
+++ b/Engine/Source/Runtime/Media/Public/IMediaPlayer.h
@@ -318,20 +318,33 @@ public:
 		return false;
 	}
 
-	class IAsyncDestructNotification
+	class IAsyncResourceReleaseNotification
 	{
 	public:
-		virtual ~IAsyncDestructNotification() {}
-		virtual void Signal() = 0;
+		virtual ~IAsyncResourceReleaseNotification() {}
+		virtual void Signal(uint32 ResourceFlags) = 0;
 	};
-	typedef TSharedRef<IAsyncDestructNotification, ESPMode::ThreadSafe> IAsyncDestructNotificationRef;
+	typedef TSharedRef<IAsyncResourceReleaseNotification, ESPMode::ThreadSafe> IAsyncResourceReleaseNotificationRef;
 
-	virtual bool SetAsyncDestructionNotification(IAsyncDestructNotificationRef AsyncDestructNotification)
+	/**
+	 * Set async resource release notification for use with IMediaPlayerLifecycleManagerDelegate
+	 */
+	virtual bool SetAsyncResourceReleaseNotification(IAsyncResourceReleaseNotificationRef AsyncDestructNotification)
 	{
 		// Override in child class if needed.
 		return false;
 	}
 
+	/*
+	* Return IMediaPlayerLifecycleManagerDelegate::ResourceFlags bitmask to indicate resource types recreated on a open call
+	*/
+	virtual uint32 GetNewResourcesOnOpen() const
+	{
+		// Override in child class if needed. Default assumes resources carry over (to a large degree) and are created per instance only
+		// (as a simplified model for older players)
+		return 0;
+	}
+
 public:
 
 	/** Virtual destructor. */
diff --git a/Engine/Source/Runtime/MediaUtils/Private/MediaPlayerFacade.cpp b/Engine/Source/Runtime/MediaUtils/Private/MediaPlayerFacade.cpp
--- a/Engine/Source/Runtime/MediaUtils/Private/MediaPlayerFacade.cpp
+++ b/Engine/Source/Runtime/MediaUtils/Private/MediaPlayerFacade.cpp
@@ -678,8 +678,8 @@ bool FMediaPlayerFacade::IsReady() const
 class FMediaPlayerLifecycleManagerDelegateOpenRequest : public IMediaPlayerLifecycleManagerDelegate::IOpenRequest
 {
 public:
-	FMediaPlayerLifecycleManagerDelegateOpenRequest(const FString& InUrl, const IMediaOptions* InOptions, const FMediaPlayerOptions* InPlayerOptions, IMediaPlayerFactory* InPlayerFactory, bool bInWillCreatePlayer)
-		: Url(InUrl), Options(InOptions), PlayerFactory(InPlayerFactory), bWillCreatePlayer(bInWillCreatePlayer)
+	FMediaPlayerLifecycleManagerDelegateOpenRequest(const FString& InUrl, const IMediaOptions* InOptions, const FMediaPlayerOptions* InPlayerOptions, IMediaPlayerFactory* InPlayerFactory, bool bInWillCreatePlayer, uint32 InWillUseNewResources)
+		: Url(InUrl), Options(InOptions), PlayerFactory(InPlayerFactory), bWillCreatePlayer(bInWillCreatePlayer), NewResources(InWillUseNewResources)
 	{
 		if (InPlayerOptions)
 		{
@@ -712,14 +712,20 @@ public:
 		return bWillCreatePlayer;
 	}
 
+	virtual bool WillUseNewResources(uint32 ResourceFlags) const
+	{
+		return !!(NewResources & ResourceFlags);
+	}
+
 	FString Url;
 	const IMediaOptions* Options;
 	TOptional<FMediaPlayerOptions> PlayerOptions;
 	IMediaPlayerFactory* PlayerFactory;
 	bool bWillCreatePlayer;
+	uint32 NewResources;
 };
 
-class FMediaPlayerLifecycleManagerDelegateControl : public IMediaPlayerLifecycleManagerDelegate::IControl
+class FMediaPlayerLifecycleManagerDelegateControl : public IMediaPlayerLifecycleManagerDelegate::IControl, public TSharedFromThis<FMediaPlayerLifecycleManagerDelegateControl, ESPMode::ThreadSafe>
 {
 public:
 	FMediaPlayerLifecycleManagerDelegateControl(TWeakPtr<FMediaPlayerFacade, ESPMode::ThreadSafe> InFacade) : Facade(InFacade), InstanceID(~0), SubmittedRequest(false) {}
@@ -740,12 +746,13 @@ public:
 		if (TSharedPtr<FMediaPlayerFacade, ESPMode::ThreadSafe> PinnedFacade = Facade.Pin())
 		{
 			const FMediaPlayerLifecycleManagerDelegateOpenRequest* OR = static_cast<const FMediaPlayerLifecycleManagerDelegateOpenRequest*>(OpenRequest.Get());
-			bool bOk = PinnedFacade->ContinueOpen(OR->Url, OR->Options, OR->PlayerOptions.IsSet() ? &OR->PlayerOptions.GetValue() : nullptr, OR->PlayerFactory, OR->bWillCreatePlayer);
-			if (bOk)
+			if (PinnedFacade->ContinueOpen(AsShared(), OR->Url, OR->Options, OR->PlayerOptions.IsSet() ? &OR->PlayerOptions.GetValue() : nullptr, OR->PlayerFactory, OR->bWillCreatePlayer, InstanceID))
 			{
 				SubmittedRequest = true;
 			}
-			return bOk;
+			//note: we return "true" in all cases in which we were able to get to call "ContinueOpen". Failures in here will be messaged to the delegate using the OnMediaPlayerCreateFailed() method
+			// (returning true here allows for capturing an unlikely early death of the facade while protecting us from double-handling the failure of the creation in the delegate)
+			return true;
 		}
 		return false;
 	}
@@ -765,6 +772,11 @@ public:
 		InstanceID = InInstanceID;
 	}
 
+	void Reset()
+	{
+		SubmittedRequest = true;
+	}
+
 private:
 	TWeakPtr<FMediaPlayerFacade, ESPMode::ThreadSafe> Facade;
 	uint64 InstanceID;
@@ -773,20 +785,27 @@ private:
 
 // ----------------------------------------------------------------------------------------------------------------------------------------------
 
-bool FMediaPlayerFacade::NotifyLifetimeManagerDelegate_PlayerOpen(const FString& Url, const IMediaOptions* Options, const FMediaPlayerOptions* PlayerOptions, IMediaPlayerFactory* PlayerFactory, bool bWillCreatePlayer)
+bool FMediaPlayerFacade::NotifyLifetimeManagerDelegate_PlayerOpen(IMediaPlayerLifecycleManagerDelegate::IControlRef & NewLifecycleManagerDelegateControl, const FString& Url, const IMediaOptions* Options, const FMediaPlayerOptions* PlayerOptions, IMediaPlayerFactory* PlayerFactory, bool bWillCreatePlayer, uint32 WillUseNewResources, uint64 NewPlayerInstanceID)
 {
 	check(IsInGameThread());
 
 	if (IMediaPlayerLifecycleManagerDelegate* Delegate = MediaModule->GetPlayerLifecycleManagerDelegate())
 	{
-		LifecycleManagerDelegateControl = MakeShared<FMediaPlayerLifecycleManagerDelegateControl, ESPMode::ThreadSafe>(AsShared());
-		if (LifecycleManagerDelegateControl.IsValid())
+		NewLifecycleManagerDelegateControl = MakeShared<FMediaPlayerLifecycleManagerDelegateControl, ESPMode::ThreadSafe>(AsShared());
+		if (NewLifecycleManagerDelegateControl.IsValid())
 		{
-			IMediaPlayerLifecycleManagerDelegate::IOpenRequestRef OpenRequest(new FMediaPlayerLifecycleManagerDelegateOpenRequest(Url, Options, PlayerOptions, PlayerFactory, bWillCreatePlayer));
+			// Set instance ID we will use for a new player if we get the go-ahead to create it (old ID if player is about to be reused)
+			static_cast<FMediaPlayerLifecycleManagerDelegateControl*>(NewLifecycleManagerDelegateControl.Get())->SetInstanceID(NewPlayerInstanceID);
+
+			IMediaPlayerLifecycleManagerDelegate::IOpenRequestRef OpenRequest(new FMediaPlayerLifecycleManagerDelegateOpenRequest(Url, Options, PlayerOptions, PlayerFactory, bWillCreatePlayer, WillUseNewResources));
 			if (OpenRequest.IsValid())
 			{
-				return Delegate->OnMediaPlayerOpen(LifecycleManagerDelegateControl, OpenRequest);
+				if (Delegate->OnMediaPlayerOpen(NewLifecycleManagerDelegateControl, OpenRequest))
+				{
+					return true;
+				}
 			}
+			static_cast<FMediaPlayerLifecycleManagerDelegateControl*>(NewLifecycleManagerDelegateControl.Get())->Reset();
 		}
 	}
 	return false;
@@ -799,8 +818,6 @@ bool FMediaPlayerFacade::NotifyLifetimeManagerDelegate_PlayerCreated()
 
 	if (LifecycleManagerDelegateControl.IsValid())
 	{
-		static_cast<FMediaPlayerLifecycleManagerDelegateControl*>(LifecycleManagerDelegateControl.Get())->SetInstanceID(PlayerInstanceID);
-
 		if (IMediaPlayerLifecycleManagerDelegate* Delegate = MediaModule->GetPlayerLifecycleManagerDelegate())
 		{
 			Delegate->OnMediaPlayerCreated(LifecycleManagerDelegateControl);
@@ -810,6 +827,21 @@ bool FMediaPlayerFacade::NotifyLifetimeManagerDelegate_PlayerCreated()
 	return false;
 }
 
+bool FMediaPlayerFacade::NotifyLifetimeManagerDelegate_PlayerCreateFailed()
+{
+	check(IsInGameThread());
+
+	if (LifecycleManagerDelegateControl.IsValid())
+	{
+		if (IMediaPlayerLifecycleManagerDelegate* Delegate = MediaModule->GetPlayerLifecycleManagerDelegate())
+		{
+			Delegate->OnMediaPlayerCreateFailed(LifecycleManagerDelegateControl);
+			return true;
+		}
+	}
+	return false;
+}
+
 bool FMediaPlayerFacade::NotifyLifetimeManagerDelegate_PlayerClosed()
 {
 	check(IsInGameThread());
@@ -840,6 +872,21 @@ bool FMediaPlayerFacade::NotifyLifetimeManagerDelegate_PlayerDestroyed()
 	return false;
 }
 
+bool FMediaPlayerFacade::NotifyLifetimeManagerDelegate_PlayerResourcesReleased(uint32 ResourceFlags)
+{
+	check(IsInGameThread());
+
+	if (LifecycleManagerDelegateControl.IsValid())
+	{
+		if (IMediaPlayerLifecycleManagerDelegate* Delegate = MediaModule->GetPlayerLifecycleManagerDelegate())
+		{
+			Delegate->OnMediaPlayerResourcesReleased(LifecycleManagerDelegateControl, ResourceFlags);
+			return true;
+		}
+	}
+	return false;
+}
+
 // ----------------------------------------------------------------------------------------------------------------------------------------------
 
 void FMediaPlayerFacade::DestroyPlayer()
@@ -851,34 +898,11 @@ void FMediaPlayerFacade::DestroyPlayer()
 		return;
 	}
 
-	class FAsyncDestructNotification : public IMediaPlayer::IAsyncDestructNotification
-	{
-	public:
-		FAsyncDestructNotification(IMediaModule* InMediaModule, IMediaPlayerLifecycleManagerDelegate::IControlRef InDelegateControl) : MediaModule(InMediaModule), DelegateControl(InDelegateControl) {}
-
-		virtual void Signal() override
-		{
-			TFunction<void()> NotifyTask = [TargetMediaModule{ MoveTemp(MediaModule) }, TargetDelegateControl { MoveTemp(DelegateControl) }]()
-			{
-				if (IMediaPlayerLifecycleManagerDelegate* Delegate = TargetMediaModule->GetPlayerLifecycleManagerDelegate())
-				{
-					Delegate->OnMediaPlayerDestroyed(TargetDelegateControl);
-				}
-			};
-			Async(EAsyncExecution::TaskGraphMainThread, NotifyTask);
-		};
-
-		IMediaModule* MediaModule;
-		IMediaPlayerLifecycleManagerDelegate::IControlRef DelegateControl;
-	};
-
-	bool bWillNotify = Player->SetAsyncDestructionNotification(TSharedRef<IMediaPlayer::IAsyncDestructNotification, ESPMode::ThreadSafe>(new FAsyncDestructNotification(MediaModule, LifecycleManagerDelegateControl)));
 	Player.Reset();
-	if (!bWillNotify)
+	NotifyLifetimeManagerDelegate_PlayerDestroyed();
+	if (!PlayerUsesResourceReleaseNotification)
 	{
-		// The player did not implement the async destruction notification, so we assume it did everything synchronously...
-		// (note: there is a chance something else is hanging on to the player's ref and we are triggering this notification too early, but as we expect all V2 players to implement this eventually, we accept this for now)
-		NotifyLifetimeManagerDelegate_PlayerDestroyed();
+		NotifyLifetimeManagerDelegate_PlayerResourcesReleased(IMediaPlayerLifecycleManagerDelegate::ResourceFlags_All);
 	}
 }
 
@@ -908,18 +932,40 @@ bool FMediaPlayerFacade::Open(const FString& Url, const IMediaOptions* Options,
 	IMediaPlayerFactory* OldFactory(Player.IsValid() ? MediaModule->GetPlayerFactory(Player->GetPlayerPluginGUID()) : nullptr);
 
 	bool bWillCreatePlayer = (PlayerFactory != OldFactory);
+	uint64 NewPlayerInstanceID;
+	uint32 WillUseNewResources;
+
+	if (bWillCreatePlayer)
+	{
+		NewPlayerInstanceID = MediaModule->CreateMediaPlayerInstanceID();
+		WillUseNewResources = IMediaPlayerLifecycleManagerDelegate::ResourceFlags_All; // as we create a new player we assume all resources a newly created in any case
+	}
+	else
+	{
+		if (Player.IsValid())
+		{
+			NewPlayerInstanceID = PlayerInstanceID;
+			WillUseNewResources = Player->GetNewResourcesOnOpen(); // ask player what resources it will create again even if it already exists
+		}
+		else
+		{
+			NewPlayerInstanceID = ~0;
+			WillUseNewResources = 0;
+		}
+	}
 
-	if (FMediaPlayerFacade::NotifyLifetimeManagerDelegate_PlayerOpen(Url, Options, PlayerOptions, PlayerFactory, bWillCreatePlayer))
+	IMediaPlayerLifecycleManagerDelegate::IControlRef NewLifecycleManagerDelegateControl;
+	if (FMediaPlayerFacade::NotifyLifetimeManagerDelegate_PlayerOpen(NewLifecycleManagerDelegateControl, Url, Options, PlayerOptions, PlayerFactory, bWillCreatePlayer, WillUseNewResources, NewPlayerInstanceID))
 	{
 		// Assume all is well: the delegate will either (have) submit(ted) the request or not -- in any case we need to assume the best -> "true"
 		return true;
 	}
 
 	// We did not notify successfully or the delegate will not submit the request in its own. Do so here...
-	return ContinueOpen(Url, Options, PlayerOptions, PlayerFactory, bWillCreatePlayer);
+	return ContinueOpen(NewLifecycleManagerDelegateControl, Url, Options, PlayerOptions, PlayerFactory, bWillCreatePlayer, NewPlayerInstanceID);
 }
 
-bool FMediaPlayerFacade::ContinueOpen(const FString& Url, const IMediaOptions* Options, const FMediaPlayerOptions* PlayerOptions, IMediaPlayerFactory* PlayerFactory, bool bCreateNewPlayer)
+bool FMediaPlayerFacade::ContinueOpen(IMediaPlayerLifecycleManagerDelegate::IControlRef NewLifecycleManagerDelegateControl, const FString& Url, const IMediaOptions* Options, const FMediaPlayerOptions* PlayerOptions, IMediaPlayerFactory* PlayerFactory, bool bCreateNewPlayer, uint64 NewPlayerInstanceID)
 {
 	// Create or reuse player
 	TSharedPtr<IMediaPlayer, ESPMode::ThreadSafe> NewPlayer(bCreateNewPlayer ? PlayerFactory->CreatePlayer(*this) : Player);
@@ -930,13 +976,41 @@ bool FMediaPlayerFacade::ContinueOpen(const FString& Url, const IMediaOptions* O
 	{
 		DestroyPlayer();
 
+		class FAsyncResourceReleaseNotification : public IMediaPlayer::IAsyncResourceReleaseNotification
+		{
+		public:
+			FAsyncResourceReleaseNotification(IMediaModule* InMediaModule, IMediaPlayerLifecycleManagerDelegate::IControlRef InDelegateControl) : MediaModule(InMediaModule), DelegateControl(InDelegateControl) {}
+
+			virtual void Signal(uint32 ResourceFlags) override
+			{
+				TFunction<void()> NotifyTask = [TargetMediaModule=MediaModule, TargetDelegateControl=DelegateControl, ResourceFlags]()
+				{
+					if (IMediaPlayerLifecycleManagerDelegate* Delegate = TargetMediaModule->GetPlayerLifecycleManagerDelegate())
+					{
+						Delegate->OnMediaPlayerResourcesReleased(TargetDelegateControl, ResourceFlags);
+					}
+				};
+				Async(EAsyncExecution::TaskGraphMainThread, NotifyTask);
+			};
+
+			IMediaModule* MediaModule;
+			IMediaPlayerLifecycleManagerDelegate::IControlRef DelegateControl;
+		};
+
 		FScopeLock Lock(&CriticalSection);
 		Player = NewPlayer;
-		PlayerInstanceID = MediaModule->CreateMediaPlayerInstanceID();
+		PlayerInstanceID = NewPlayerInstanceID;
+		LifecycleManagerDelegateControl = NewLifecycleManagerDelegateControl;
+		PlayerUsesResourceReleaseNotification = LifecycleManagerDelegateControl.IsValid() ? Player->SetAsyncResourceReleaseNotification(TSharedRef<IMediaPlayer::IAsyncResourceReleaseNotification, ESPMode::ThreadSafe>(new FAsyncResourceReleaseNotification(MediaModule, LifecycleManagerDelegateControl))) : false;
+	}
+	else
+	{
+		LifecycleManagerDelegateControl = NewLifecycleManagerDelegateControl;
 	}
 
 	if (!Player.IsValid())
 	{
+		NotifyLifetimeManagerDelegate_PlayerCreateFailed();
 		// Make sure we don't get called from the "tickable" thread anymore - no need as we have no player
 		MediaModule->GetTicker().RemoveTickable(AsShared());
 		return false;
@@ -959,6 +1033,7 @@ bool FMediaPlayerFacade::ContinueOpen(const FString& Url, const IMediaOptions* O
 	// open the new media source
 	if (!Player->Open(Url, Options, PlayerOptions))
 	{
+		NotifyLifetimeManagerDelegate_PlayerCreateFailed();
 		CurrentUrl.Empty();
 		ActivePlayerOptions.Reset();
 
@@ -970,7 +1045,7 @@ bool FMediaPlayerFacade::ContinueOpen(const FString& Url, const IMediaOptions* O
 	LastAudioSampleProcessedTime.Invalidate();
 	CurrentFrameAudioTimeStamp.Invalidate();
 
-	if (Player.IsValid() && bCreateNewPlayer)
+	if (bCreateNewPlayer)
 	{
 		NotifyLifetimeManagerDelegate_PlayerCreated();
 	}
diff --git a/Engine/Source/Runtime/MediaUtils/Public/IMediaPlayerLifecycleManager.h b/Engine/Source/Runtime/MediaUtils/Public/IMediaPlayerLifecycleManager.h
--- a/Engine/Source/Runtime/MediaUtils/Public/IMediaPlayerLifecycleManager.h
+++ b/Engine/Source/Runtime/MediaUtils/Public/IMediaPlayerLifecycleManager.h
@@ -27,18 +27,60 @@ struct FMediaPlayerOptions;
 
 	Notes:
 
-	- All notification callbacks are issues on the game thread
-	- The PlayerOpen callback will be issued before any player instance has been created, hence the ID is invalid (but a facade may exist)
-	-- If the callback returns "true" it must take care to call SubmitOpenRequest() with the passed along open request as soon as it deems it "ok" for the player to serve the request to be created
-	-- (if not "submitted" the open will not create an actual player; if "false" is returned the default creation method will continue as if a submit was done)
-	- GetFacade() may return a null pointer in case the facade instance was destroyed since the event was triggered and hence should always be checked before use
-	- The PlayerDestroyed callback will be triggered after the player is already destroyed. The facade may or may not still exist. The ID is still valid so this event can be properly tracked
+	- All notification callbacks are issued on the game thread
+	- Calls into the system can cause further callbacks to be issued even while the call is ongoing
+	- IControl is the interface allowing access to player data and instance as well as submitting OpenRequests as needed
+	- IOpenRequest is passed into the OnMediaPlayerOpen() callback when a new player is about to be opened / created
+	  It contains all information available about the request so that the code can decided to let the request through, delay or deny it.
+	  - To just let it pass through return false from the callback
+	  - To delay or deny it return true
+	  - Keep a reference to the IControl passed in until you are ready to submit the open request or have decided to deny it
+	  - To deny it simply do not submit the request and release all references (this will appear to the rest if the system as a failure to open the player)
+	- An InstanceID is available on all callbacks, but be aware that the ID is not yet assigned to an actual player during OnMediaPlayerOpen
+	- The OnMediaPlayerOpen can be called for a certain InstanceID more then once if the player is reused without being destroyed first
+	  - This can lead to the need to track >1 resource use cases for a single ID (e.g. some players will only release the last set of resources once the new resources are already allocated for a while to avoid stalls)
+	- Some callbacks may be issued at a time that the player facade associated with the player (and may be even the low level player) do no longer exist. Be prepared for e.g. GetFacade() to return a null reference
+
+	OnMediaPlayerOpen
+
+	Called when a player is about to be opened. If this is the first time the facade receives an open call or if the player type changed this is called before the actual creation of the player.
+	In any case this is called early enough to still stop the open from proceeding. Aside of tracking player activity, the ability to control the processing of open calls is the main purpose of this callback.
+
+	OnMediaPlayerCreated
+
+	Called when the creation of a new player succeeded. Note that the open may still fail if the player later on has issues playing the content it is pointed to.
+
+	OnMediaPlayerCreateFailed
+
+	Creation of the player failed. This will posisbly be called from within the SubmitOpenRequest() call of the IControl interface. It signifies that either the player factory did not manage to create a player
+	or that the low level player failed to open in the synchronous part of its opening process.
+
+	OnMediaPlayerClosed
+
+	The low level player closed the current playback session.
+
+	OnMediaPlayerDestroyed
+
+	The low level player instance just got destroyed.
+
+	OnMediaPlayerResourcesReleased
+
+	The low level player instance release resources. This callback receives a bitmask of possible resource types (Resource_Flag...) to allow proper tracking of players that release certain types of resources at different times.
+	Some players will issue this callback asynchronously, it may hence be triggered long (100-200ms are possible) after the closing or destruction of a player.
 
 */
 //! Interface to receive global player lifetime events from media framework
 class IMediaPlayerLifecycleManagerDelegate
 {
 public:
+	enum {
+		ResourceFlags_Decoder = 1 << 0,
+		ResourceFlags_OutputBuffers = 1 << 1,
+
+		ResourceFlags_All = (1 << 2) - 1,
+		ResourceFlags_Any = ResourceFlags_All,
+	};
+
 	//! Request to create and open a player
 	class IOpenRequest
 	{
@@ -50,6 +92,7 @@ public:
 		virtual const FMediaPlayerOptions* GetPlayerOptions() const = 0;
 		virtual IMediaPlayerFactory* GetPlayerFactory() const = 0;
 		virtual bool WillCreateNewPlayer() const = 0;
+		virtual bool WillUseNewResources(uint32 ResourceFlags) const = 0;
 	};
 	typedef TSharedPtr<IOpenRequest, ESPMode::ThreadSafe> IOpenRequestRef;
 
@@ -70,6 +113,8 @@ public:
 
 	virtual bool OnMediaPlayerOpen(IControlRef Control, IOpenRequestRef OpenRequest) = 0;
 	virtual void OnMediaPlayerCreated(IControlRef Control) = 0;
+	virtual void OnMediaPlayerCreateFailed(IControlRef Control) = 0;
 	virtual void OnMediaPlayerClosed(IControlRef Control) = 0;
 	virtual void OnMediaPlayerDestroyed(IControlRef Control) = 0;
+	virtual void OnMediaPlayerResourcesReleased(IControlRef Control, uint32 ResourceFlags) = 0;
 };
diff --git a/Engine/Source/Runtime/MediaUtils/Public/MediaPlayerFacade.h b/Engine/Source/Runtime/MediaUtils/Public/MediaPlayerFacade.h
--- a/Engine/Source/Runtime/MediaUtils/Public/MediaPlayerFacade.h
+++ b/Engine/Source/Runtime/MediaUtils/Public/MediaPlayerFacade.h
@@ -780,10 +780,12 @@ private:
 	friend class FMediaPlayerLifecycleManagerDelegateControl;
 
 	// Internal
-	bool NotifyLifetimeManagerDelegate_PlayerOpen(const FString& InUrl, const IMediaOptions* Options, const FMediaPlayerOptions* InPlayerOptions, IMediaPlayerFactory* PlayerFactory, bool bWillCreatePlayer);
+	bool NotifyLifetimeManagerDelegate_PlayerOpen(IMediaPlayerLifecycleManagerDelegate::IControlRef& NewLifecycleManagerDelegateControl, const FString& InUrl, const IMediaOptions* Options, const FMediaPlayerOptions* InPlayerOptions, IMediaPlayerFactory* PlayerFactory, bool bWillCreatePlayer, uint32 WillUseNewResources, uint64 NewPlayerInstanceID);
 	bool NotifyLifetimeManagerDelegate_PlayerCreated();
+	bool NotifyLifetimeManagerDelegate_PlayerCreateFailed();
 	bool NotifyLifetimeManagerDelegate_PlayerClosed();
 	bool NotifyLifetimeManagerDelegate_PlayerDestroyed();
+	bool NotifyLifetimeManagerDelegate_PlayerResourcesReleased(uint32 ResourceFlags);
 
 	void ProcessVideoSamples(IMediaSamples& Samples, const TRange<FMediaTimeStamp> & TimeRange);
 	void ProcessCaptionSamples(IMediaSamples& Samples, TRange<FMediaTimeStamp> TimeRange);
@@ -793,7 +795,7 @@ private:
 
 	void DestroyPlayer();
 
-	bool ContinueOpen(const FString& Url, const IMediaOptions* Options, const FMediaPlayerOptions* PlayerOptions, IMediaPlayerFactory* PlayerFactory, bool bCreateNewPlayer);
+	bool ContinueOpen(IMediaPlayerLifecycleManagerDelegate::IControlRef NewLifecycleManagerDelegateControl, const FString& Url, const IMediaOptions* Options, const FMediaPlayerOptions* PlayerOptions, IMediaPlayerFactory* PlayerFactory, bool bCreateNewPlayer, uint64 NewPlayerInstanceID);
 
 	/** Audio sample sinks. */
 	TWeakPtr<FMediaAudioSampleSink, ESPMode::ThreadSafe> PrimaryAudioSink;
@@ -846,6 +848,9 @@ private:
 
 	/** Low level player instance ID */
 	uint64 PlayerInstanceID = ~0;
+	
+	/** Low level player will use callback to notify of resource release */
+	bool PlayerUsesResourceReleaseNotification = false;
 
 	/** Media player Guid */
 	FGuid PlayerGuid;
