MediaUtils: added missing media events to EventToString()

#rnx

diff --git a/Engine/Source/Runtime/MediaUtils/Private/MediaHelpers.cpp b/Engine/Source/Runtime/MediaUtils/Private/MediaHelpers.cpp
--- a/Engine/Source/Runtime/MediaUtils/Private/MediaHelpers.cpp
+++ b/Engine/Source/Runtime/MediaUtils/Private/MediaHelpers.cpp
@@ -25,6 +25,9 @@ namespace MediaUtils
 		case EMediaEvent::TracksChanged: return TEXT("TracksChanged");
 
 		case EMediaEvent::Internal_PurgeVideoSamplesHint: return TEXT("Internal: PurgeVideoSamplesHint");
+		case EMediaEvent::Internal_ResetForDiscontinuity: return TEXT("Internal: ResetForDiscontinuity");
+		case EMediaEvent::Internal_RenderClockStart: return TEXT("Internal: RenderClockStart");
+		case EMediaEvent::Internal_RenderClockStop: return TEXT("Internal: RenderClockStop");
 		case EMediaEvent::Internal_VideoSamplesAvailable: return TEXT("Internal: VideoSamplesAvailable");
 		case EMediaEvent::Internal_VideoSamplesUnavailable: return TEXT("Internal: VideoSamplesUnavailable");
 		case EMediaEvent::Internal_AudioSamplesAvailable: return TEXT("Internal: AudioSamplesAvailable");
