Fixing up a possible crash during shutdown if media playback using the lifecycle manager should be active at the time
#rnx

diff --git a/Engine/Source/Runtime/MediaUtils/Private/MediaPlayerFacade.cpp b/Engine/Source/Runtime/MediaUtils/Private/MediaPlayerFacade.cpp
--- a/Engine/Source/Runtime/MediaUtils/Private/MediaPlayerFacade.cpp
+++ b/Engine/Source/Runtime/MediaUtils/Private/MediaPlayerFacade.cpp
@@ -979,15 +979,22 @@ bool FMediaPlayerFacade::ContinueOpen(IMediaPlayerLifecycleManagerDelegate::ICon
 		class FAsyncResourceReleaseNotification : public IMediaPlayer::IAsyncResourceReleaseNotification
 		{
 		public:
-			FAsyncResourceReleaseNotification(IMediaModule* InMediaModule, IMediaPlayerLifecycleManagerDelegate::IControlRef InDelegateControl) : MediaModule(InMediaModule), DelegateControl(InDelegateControl) {}
+			FAsyncResourceReleaseNotification(IMediaPlayerLifecycleManagerDelegate::IControlRef InDelegateControl) : DelegateControl(InDelegateControl) {}
 
 			virtual void Signal(uint32 ResourceFlags) override
 			{
-				TFunction<void()> NotifyTask = [TargetMediaModule=MediaModule, TargetDelegateControl=DelegateControl, ResourceFlags]()
+				TFunction<void()> NotifyTask = [TargetDelegateControl=DelegateControl, ResourceFlags]()
 				{
-					if (IMediaPlayerLifecycleManagerDelegate* Delegate = TargetMediaModule->GetPlayerLifecycleManagerDelegate())
+					// Get MediaModule & check if it is already unloaded...
+					IMediaModule* TargetMediaModule = FModuleManager::GetModulePtr<IMediaModule>("Media");
+					if (TargetMediaModule)
 					{
-						Delegate->OnMediaPlayerResourcesReleased(TargetDelegateControl, ResourceFlags);
+						// Delegate still there?
+						if (IMediaPlayerLifecycleManagerDelegate* Delegate = TargetMediaModule->GetPlayerLifecycleManagerDelegate())
+						{
+							// Notify it!
+							Delegate->OnMediaPlayerResourcesReleased(TargetDelegateControl, ResourceFlags);
+						}
 					}
 				};
 				Async(EAsyncExecution::TaskGraphMainThread, NotifyTask);
@@ -1001,7 +1008,7 @@ bool FMediaPlayerFacade::ContinueOpen(IMediaPlayerLifecycleManagerDelegate::ICon
 		Player = NewPlayer;
 		PlayerInstanceID = NewPlayerInstanceID;
 		LifecycleManagerDelegateControl = NewLifecycleManagerDelegateControl;
-		PlayerUsesResourceReleaseNotification = LifecycleManagerDelegateControl.IsValid() ? Player->SetAsyncResourceReleaseNotification(TSharedRef<IMediaPlayer::IAsyncResourceReleaseNotification, ESPMode::ThreadSafe>(new FAsyncResourceReleaseNotification(MediaModule, LifecycleManagerDelegateControl))) : false;
+		PlayerUsesResourceReleaseNotification = LifecycleManagerDelegateControl.IsValid() ? Player->SetAsyncResourceReleaseNotification(TSharedRef<IMediaPlayer::IAsyncResourceReleaseNotification, ESPMode::ThreadSafe>(new FAsyncResourceReleaseNotification(LifecycleManagerDelegateControl))) : false;
 	}
 	else
 	{
