Pulling in changes to make blocked playback work with "V2 players" in MFW

diff --git a/Engine/Source/Runtime/Media/Public/IMediaControls.h b/Engine/Source/Runtime/Media/Public/IMediaControls.h
--- a/Engine/Source/Runtime/Media/Public/IMediaControls.h
+++ b/Engine/Source/Runtime/Media/Public/IMediaControls.h
@@ -209,6 +209,17 @@ public:
 	 */
 	virtual bool SetRate(float Rate) = 0;
 
+	/**
+	 * Hint for player indicating that blocked playback mode will be used / not used
+	 * 
+	 * @param bFacadeWillUseBlockingPlayback True if blocked playback will be used, false otherwiese
+	 * 
+	 * @note Implement as needed by the player.
+	 */
+	virtual void SetBlockingPlaybackHint(bool bFacadeWillUseBlockingPlayback)
+	{
+	}
+
 public:
 
 	/**
diff --git a/Engine/Source/Runtime/Media/Public/IMediaTimeSource.h b/Engine/Source/Runtime/Media/Public/IMediaTimeSource.h
--- a/Engine/Source/Runtime/Media/Public/IMediaTimeSource.h
+++ b/Engine/Source/Runtime/Media/Public/IMediaTimeSource.h
@@ -3,6 +3,7 @@
 #pragma once
 
 #include "Misc/Timespan.h"
+#include "Math/NumericLimits.h"
 
 /*
 	Timestamp value for media playback
@@ -44,7 +45,12 @@ public:
 	FMediaTimeStamp operator + (const FTimespan & Other) const { return FMediaTimeStamp(Time + Other, SequenceIndex); }
 	FMediaTimeStamp operator - (const FTimespan & Other) const { return FMediaTimeStamp(Time - Other, SequenceIndex); }
 
-	FMediaTimeStamp operator - (const FMediaTimeStamp & Other) const { return FMediaTimeStamp(Time - Other.Time, -1); }
+	FMediaTimeStamp operator - (const FMediaTimeStamp & Other) const { return FMediaTimeStamp(Time - Other.Time, MAX_int64); }
+
+	FMediaTimeStamp& operator += (const FTimespan & Other) { Time += Other; return *this; }
+	FMediaTimeStamp& operator -= (const FTimespan & Other) { Time -= Other; return *this; }
+
+	FMediaTimeStamp& operator -= (const FMediaTimeStamp & Other) { Time -= Other.Time; SequenceIndex = MAX_int64; return *this; }
 
 	FTimespan Time;
 	int64 SequenceIndex;
diff --git a/Engine/Source/Runtime/MediaAssets/Private/Assets/MediaPlayer.cpp b/Engine/Source/Runtime/MediaAssets/Private/Assets/MediaPlayer.cpp
--- a/Engine/Source/Runtime/MediaAssets/Private/Assets/MediaPlayer.cpp
+++ b/Engine/Source/Runtime/MediaAssets/Private/Assets/MediaPlayer.cpp
@@ -639,6 +639,13 @@ void UMediaPlayer::SetBlockOnTime(const FTimespan& Time)
 }
 
 
+void UMediaPlayer::SetBlockOnTimeRange(const TRange<FTimespan>& TimeRange)
+{
+	UE_LOG(LogMediaAssets, VeryVerbose, TEXT("%s.SetBlockOnRange %s"), *GetFName().ToString(), *TimeRange.GetLowerBoundValue().ToString(TEXT("%h:%m:%s.%t")));
+	return PlayerFacade->SetBlockOnTimeRange(TimeRange);
+}
+
+
 void UMediaPlayer::SetDesiredPlayerName(FName PlayerName)
 {
 	UE_LOG(LogMediaAssets, Verbose, TEXT("%s.SetDesiredPlayerName %s"), *GetFName().ToString(), *PlayerName.ToString());
diff --git a/Engine/Source/Runtime/MediaAssets/Public/MediaPlayer.h b/Engine/Source/Runtime/MediaAssets/Public/MediaPlayer.h
--- a/Engine/Source/Runtime/MediaAssets/Public/MediaPlayer.h
+++ b/Engine/Source/Runtime/MediaAssets/Public/MediaPlayer.h
@@ -723,7 +723,7 @@ public:
 	/**
 	 * Set the time on which to block.
 	 *
-	 * If set, this player will block in TickFetch until the video sample
+	 * If set, this player will block in TickInput or TickFetch until the video sample
 	 * for the specified time are actually available.
 	 *
 	 * @param Time The time to block on, or FTimespan::MinValue to disable.
@@ -731,6 +731,16 @@ public:
 	UFUNCTION(BlueprintCallable, Category = "Media|MediaPlayer")
 	void SetBlockOnTime(const FTimespan& Time);
 
+	/**
+	 * Set the time range on which to block.
+	 *
+	 * If set, this player will block in TickInput or TickFetch until the video sample
+	 * for the specified time are actually available.
+	 *
+	 * @param TimeRange The time range to block on, or empty range to disable.
+	 */
+	void SetBlockOnTimeRange(const TRange<FTimespan>& TimeRange);
+
 	/**
 	 * Set the name of the desired native player.
 	 *
diff --git a/Engine/Source/Runtime/MediaUtils/Private/MediaPlayerFacade.cpp b/Engine/Source/Runtime/MediaUtils/Private/MediaPlayerFacade.cpp
--- a/Engine/Source/Runtime/MediaUtils/Private/MediaPlayerFacade.cpp
+++ b/Engine/Source/Runtime/MediaUtils/Private/MediaPlayerFacade.cpp
@@ -1,4 +1,5 @@
 // Copyright Epic Games, Inc. All Rights Reserved.
+// Copyright Epic Games, Inc. All Rights Reserved.
 
 #include "MediaPlayerFacade.h"
 #include "MediaUtilsPrivate.h"
@@ -80,25 +81,30 @@ const double kOutdatedVideoSamplesTollerance = 0.050; // seconds video samples a
 
 namespace MediaPlayerFacade
 {
-	// @todo gmp: make these configurable in settings?
-
 	const FTimespan AudioPreroll = FTimespan::FromSeconds(1.0);
 	const FTimespan MetadataPreroll = FTimespan::FromSeconds(1.0);
 }
 
+static FTimespan WrappedModulo(FTimespan Time, FTimespan Duration)
+{
+	return (Time >= FTimespan::Zero()) ? (Time % Duration) : (Duration + (Time % Duration));
+}
 
 /* FMediaPlayerFacade structors
 *****************************************************************************/
 
 FMediaPlayerFacade::FMediaPlayerFacade()
 	: TimeDelay(FTimespan::Zero())
-	, BlockOnTime(FTimespan::MinValue())
 	, Cache(new FMediaSampleCache)
 	, LastRate(0.0f)
+	, CurrentRate(0.0f)
 	, bHaveActiveAudio(false)
 	, VideoSampleAvailability(-1)
 	, AudioSampleAvailability(-1)
 {
+	BlockOnRange = TRange<FMediaTimeStamp>::Empty();
+	BlockOnRangeDisabled = false;
+
 	MediaModule = FModuleManager::LoadModulePtr<IMediaModule>("Media");
 	bIsSinkFlushPending = false;
 	bDidRecentPlayerHaveError = false;
@@ -244,7 +250,8 @@ void FMediaPlayerFacade::Close()
 		NotifyLifetimeManagerDelegate_PlayerClosed();
 	}
 
-	BlockOnTime = FTimespan::MinValue();
+	BlockOnRange = TRange<FMediaTimeStamp>::Empty();
+
 	Cache->Empty();
 	CurrentUrl.Empty();
 	LastRate = 0.0f;
@@ -453,7 +460,8 @@ FMediaTimeStamp FMediaPlayerFacade::GetTimeStamp() const
 
 	if (!CurrentPlayer->GetPlayerFeatureFlag(IMediaPlayer::EFeatureFlag::UsePlaybackTimingV2))
 	{
-		return FMediaTimeStamp(); // no media opened / old player timing mode
+		// Make sure we can return values for V1 players...
+		return FMediaTimeStamp(GetTime());
 	}
 
 	// Check if there are video samples present or presence is unknown.
@@ -465,7 +473,7 @@ FMediaTimeStamp FMediaPlayerFacade::GetTimeStamp() const
 		/*
 			Returning the precise time of the sample returned during TickFetch()
 		*/
-		return LastVideoSampleProcessedTime.TimeStamp;
+		return CurrentFrameVideoTimeStamp;
 	}
 	else if (HaveAudioPlayback())
 	{
@@ -1048,10 +1056,16 @@ bool FMediaPlayerFacade::ContinueOpen(IMediaPlayerLifecycleManagerDelegate::ICon
 	}
 
 	FScopeLock Lock(&LastTimeValuesCS);
-	LastVideoSampleProcessedTime.Invalidate();
+	LastBlockOnRange = TRange<FTimespan>::Empty();
+	OnBlockSeqIndex = 0;
+	BlockOnRangeDisabled = false;
+	LastVideoSampleProcessedTimeRange = TRange<FMediaTimeStamp>::Empty();
 	LastAudioSampleProcessedTime.Invalidate();
+	CurrentFrameVideoTimeStamp.Invalidate();
 	CurrentFrameAudioTimeStamp.Invalidate();
 
+	NextEstVideoTimeAtFrameStart.Invalidate();
+
 	if (bCreateNewPlayer)
 	{
 		NotifyLifetimeManagerDelegate_PlayerCreated();
@@ -1120,7 +1134,115 @@ bool FMediaPlayerFacade::SelectTrack(EMediaTrackType TrackType, int32 TrackIndex
 
 void FMediaPlayerFacade::SetBlockOnTime(const FTimespan& Time)
 {
-	BlockOnTime = Time;
+#if !MEDIAPLAYERFACADE_DISABLE_BLOCKING
+	if (!Player.IsValid() || !Player->GetControls().CanControl(EMediaControl::BlockOnFetch))
+	{
+		return;
+	}
+
+	if (Time == FTimespan::MinValue())
+	{
+		BlockOnRange = TRange<FMediaTimeStamp>::Empty();
+		LastBlockOnRange = TRange<FTimespan>::Empty();
+		Player->GetControls().SetBlockingPlaybackHint(false);
+	}
+	else
+	{
+		BlockOnRange.Inclusive(Time, Time);
+		LastBlockOnRange.Inclusive(Time, Time);
+		Player->GetControls().SetBlockingPlaybackHint(true);
+	}
+#endif
+}
+
+
+void FMediaPlayerFacade::SetBlockOnTimeRange(const TRange<FTimespan> & TimeRange)
+{
+#if !MEDIAPLAYERFACADE_DISABLE_BLOCKING
+	if (TimeRange.IsEmpty() || !Player.IsValid() || !Player->GetControls().CanControl(EMediaControl::BlockOnFetch))
+	{
+		LastBlockOnRange = TRange<FTimespan>::Empty();
+		BlockOnRange = TRange<FMediaTimeStamp>::Empty();
+		Player->GetControls().SetBlockingPlaybackHint(false);
+		return;
+	}
+
+	FTimespan Duration(Player->GetControls().GetDuration());
+	FTimespan Start(TimeRange.GetLowerBoundValue());
+	FTimespan End(TimeRange.GetUpperBoundValue());
+
+	if (!Player->GetControls().IsLooping())
+	{
+		// We pass in the time range as is on seq-index zero at all times - players have to reject sample output / blocking logic will detect begin outside media range
+		BlockOnRange = TRange<FMediaTimeStamp>(FMediaTimeStamp(Start, 0), FMediaTimeStamp(End, 0));
+	}
+	else
+	{
+		// If this would be called very early in the player's startup after open() we would not yet be known... that would be fatal
+		/*
+		 Should this actually happen in real-life applications, we could move the computations here into an accessor method used internally, so that this would be done
+		 only if data is processed, which would also mean: we know the duration!
+		 (Exception: live playback! --> but we would not allow blocking there anyway! (makes no sense as real life use case))
+		*/
+		check(!Duration.IsZero());
+		if (Duration.IsZero())
+		{
+			// Catch if this is called to early and reset blocking...
+			BlockOnRange = TRange<FMediaTimeStamp>::Empty();
+			Player->GetControls().SetBlockingPlaybackHint(false);
+			return;
+		}
+
+		// Modulo on the time to get it into media's range
+		// (assumes zero-start-time)
+		Start = WrappedModulo(Start, Duration);
+		End = WrappedModulo(End, Duration);
+
+		float Rate = GetUnpausedRate();
+
+		// Detect any non-monotonic movement of the range...
+		if (!LastBlockOnRange.IsEmpty())
+		{
+			if (Rate >= 0.0f)
+			{
+				FTimespan LastStart = WrappedModulo(LastBlockOnRange.GetLowerBoundValue(), Duration);
+				if (LastStart > Start)
+				{
+					++OnBlockSeqIndex;
+				}
+			}
+			else
+			{
+				FTimespan LastEnd = WrappedModulo(LastBlockOnRange.GetUpperBoundValue(), Duration);
+				if (LastEnd < End)
+				{
+					--OnBlockSeqIndex;
+				}
+			}
+		}
+
+		// Check if our range crosses the loop point (sequence boundary)
+		int64 StartIndex, EndIndex;
+		if (Rate >= 0.0)
+		{
+			StartIndex = OnBlockSeqIndex;
+			EndIndex = (Start <= End) ? OnBlockSeqIndex : (OnBlockSeqIndex + 1);
+		}
+		else
+		{
+			StartIndex = (Start <= End) ? OnBlockSeqIndex : (OnBlockSeqIndex - 1);
+			EndIndex = OnBlockSeqIndex;
+		}
+
+		// Assemble final blocking range
+		BlockOnRange = TRange<FMediaTimeStamp>(FMediaTimeStamp(Start, StartIndex), FMediaTimeStamp(End, EndIndex));
+		check(!BlockOnRange.IsEmpty());
+	}
+
+	Player->GetControls().SetBlockingPlaybackHint(!BlockOnRange.IsEmpty());
+
+	LastBlockOnRange = TimeRange;
+#endif
 }
 
 
@@ -1149,11 +1271,27 @@ void FMediaPlayerFacade::SetMediaOptions(const IMediaOptions* Options)
 
 bool FMediaPlayerFacade::SetRate(float Rate)
 {
+	// Enter CS as we change the rate which we read on the tickable thread
+	FScopeLock Lock(&CriticalSection);
+
+	// Can we set the rate at all?
 	if (!Player.IsValid() || !Player->GetControls().SetRate(Rate))
 	{
 		return false;
 	}
 
+	// Is this new rate supported?
+	if (Rate != 0.0f && !(Player->GetControls().GetSupportedRates(EMediaRateThinning::Thinned).Contains(Rate) || Player->GetControls().GetSupportedRates(EMediaRateThinning::Unthinned).Contains(Rate)))
+	{
+		return false;
+	}
+
+	if (CurrentRate == Rate)
+	{
+		// no change - just return with ok status
+		return true;
+	}
+
 	if ((LastRate * Rate) < 0.0f)
 	{
 		FlushSinks(); // direction change
@@ -1170,6 +1308,12 @@ bool FMediaPlayerFacade::SetRate(float Rate)
 		}
 	}
 
+	if (Rate != 0.0)
+	{
+		LastRate = Rate;
+	}
+	CurrentRate = Rate;
+
 	return true;
 }
 
@@ -1233,20 +1377,11 @@ FTimespan FMediaPlayerFacade::GetLastAudioRenderedSampleTime() const
 
 bool FMediaPlayerFacade::BlockOnFetch() const
 {
-#if MEDIAPLAYERFACADE_DISABLE_BLOCKING
-	return false;
-#endif
-
 	check(Player.IsValid());
 
-	if (BlockOnTime == FTimespan::MinValue())
-	{
-		return false; // no blocking requested
-	}
-
-	if (!Player->GetControls().CanControl(EMediaControl::BlockOnFetch))
+	if (BlockOnRange.IsEmpty() || !Player->GetControls().CanControl(EMediaControl::BlockOnFetch) || BlockOnRangeDisabled)
 	{
-		return false; // not supported by player plug-in
+		return false; // no blocking requested / not supported
 	}
 
 	if (IsPreparing())
@@ -1254,19 +1389,121 @@ bool FMediaPlayerFacade::BlockOnFetch() const
 		return true; // block on media opening
 	}
 
-	if (!IsPlaying() || (GetRate() < 0.0f))
+	if (!IsPlaying())
 	{
-		return false; // block only in forward play
+		// no blocking if we are not playing (e.g. paused)
+		return false;
 	}
 
-	const bool VideoReady = (VideoSampleSinks.Num() == 0) || (BlockOnTime < NextVideoSampleTime);
+	if (Player->GetPlayerFeatureFlag(IMediaPlayer::EFeatureFlag::UsePlaybackTimingV2))
+	{
+		//
+		// V2 blocking logic
+		//
+		float Rate = GetUnpausedRate();
+
+		// If the current sample "out there" is actually overlapping with the current block, we might be good with no new sample
+		if (LastVideoSampleProcessedTimeRange.Overlaps(BlockOnRange))
+		{
+			// We have no new data (else we would not even call this method), but the last sample we returned is still inside the current range -> good, but...
+			// If the next sample would already cover more of the range than the older one we would like to use that instead -> but it may well be we do not have any data about the sample yet (and would indeed LIKE to block!)
+			// So, we assume that the next sample will follow with no gap and have the same duration (a pretty good, general assumption) and check against that data to see if it would be better...
+
+			// Get last sample's time range
+			TRange<FMediaTimeStamp> LastSampleTimeRange(LastVideoSampleProcessedTimeRange);
+
+			// Compute the "theoretical" next sample range...
+			TRange<FMediaTimeStamp> NextSampleTimeRange = (Rate >= 0.0f) ? TRange<FMediaTimeStamp>(LastVideoSampleProcessedTimeRange.GetUpperBoundValue(), LastVideoSampleProcessedTimeRange.GetUpperBoundValue() + LastVideoSampleProcessedTimeRange.Size<FMediaTimeStamp>().Time)
+																		 : TRange<FMediaTimeStamp>(LastVideoSampleProcessedTimeRange.GetLowerBoundValue() - LastVideoSampleProcessedTimeRange.Size<FMediaTimeStamp>().Time, LastVideoSampleProcessedTimeRange.GetLowerBoundValue());
 
-	if (VideoReady)
+			FTimespan Duration = Player->GetControls().GetDuration();
+
+			if (!Player->GetControls().IsLooping())
+			{
+				// If we are not looping we need to clamp against the media's duration
+				// (we assume it starts at zero here!)
+				check(NextSampleTimeRange.GetLowerBoundValue().SequenceIndex == 0);
+				NextSampleTimeRange = TRange<FMediaTimeStamp>::Intersection(NextSampleTimeRange, TRange<FMediaTimeStamp>(FMediaTimeStamp(0, 0), FMediaTimeStamp(Duration, 0)));
+			}
+			else
+			{
+				if (NextSampleTimeRange.GetLowerBoundValue().Time >= Duration)
+				{
+					check(Rate >= 0.0f);
+					NextSampleTimeRange = TRange<FMediaTimeStamp>(FMediaTimeStamp(NextSampleTimeRange.GetLowerBoundValue().Time - Duration, NextSampleTimeRange.GetLowerBoundValue().SequenceIndex + 1), FMediaTimeStamp(NextSampleTimeRange.GetUpperBoundValue().Time - Duration, NextSampleTimeRange.GetUpperBoundValue().SequenceIndex + 1));
+				}
+				else if (NextSampleTimeRange.GetLowerBoundValue().Time < FTimespan::Zero())
+				{
+					check(Rate < 0.0f);
+					NextSampleTimeRange = TRange<FMediaTimeStamp>(FMediaTimeStamp(NextSampleTimeRange.GetLowerBoundValue().Time + Duration, NextSampleTimeRange.GetLowerBoundValue().SequenceIndex - 1), FMediaTimeStamp(NextSampleTimeRange.GetUpperBoundValue().Time + Duration, NextSampleTimeRange.GetUpperBoundValue().SequenceIndex - 1));
+				}
+			}
+
+			// Compute which one is larger inside the current range...
+			int64 LastSampleCoverage = TRange<FMediaTimeStamp>::Intersection(BlockOnRange, LastSampleTimeRange).Size<FMediaTimeStamp>().Time.GetTicks();
+			int64 NextSampleCoverage = TRange<FMediaTimeStamp>::Intersection(BlockOnRange, NextSampleTimeRange).Size<FMediaTimeStamp>().Time.GetTicks();
+
+			if (LastSampleCoverage > NextSampleCoverage)
+			{
+				// Last one we returned is still good. No blocking needed...
+				return false;
+			}
+		}
+
+		if (!Player->GetControls().IsLooping())
+		{
+			// Is the sample outside the media's range?
+			// (note: this assumes the media starts at time ZERO - this will not be the case at all times (e.g. life playback) -- for now we assume a player will flagged blocked playback as invalid in that case!)
+			if (BlockOnRange.GetUpperBoundValue() < FMediaTimeStamp(0) ||  Player->GetControls().GetDuration() <= BlockOnRange.GetLowerBoundValue().Time)
+			{
+				return false;
+			}
+		}
+
+		// Did we pass the range already?
+		// (this would imply a malfunction or lack of support in the player -- but: we need to catch it to avoid stalls)
+		const FMediaTimeStamp CurrentTime = GetTimeStamp();
+		if (Rate >= 0)
+		{
+			if (CurrentTime >= BlockOnRange.GetUpperBoundValue())
+			{
+				UE_LOG(LogMediaUtils, Warning, TEXT("Blocking range for playback is behind player playback position. Skipping blocking; possible malfunction."));
+				return false;
+			}
+		}
+		else
+		{
+			if (CurrentTime < BlockOnRange.GetLowerBoundValue())
+			{
+				UE_LOG(LogMediaUtils, Warning, TEXT("Blocking range for playback is behind player playback position. Skipping blocking; possible malfunction."));
+				return false;
+			}
+		}
+
+		// Block until sample arrives!
+		return true;
+	}
+	else
 	{
-		return false; // video is ready
+		//
+		// V1 blocking logic
+		//
+
+		if (CurrentRate < 0.0f)
+		{
+			return false; // block only in forward play
+		}
+
+		const bool VideoReady = (VideoSampleSinks.Num() == 0) || (BlockOnRange.GetUpperBoundValue() < NextVideoSampleTime);
+
+		if (VideoReady)
+		{
+			return false; // video is ready
+		}
+
+		return true;
 	}
 
-	return true;
 }
 
 
@@ -1289,8 +1526,10 @@ void FMediaPlayerFacade::FlushSinks()
 
 	NextVideoSampleTime = FTimespan::MinValue();
 
-	// note: we do not invalidate the LastXXXXSampleProcessingTime values here -> it is more natural for a outside caller to receive the "last good time" e.g. during a seek
 	LastAudioRenderedSampleTime.Invalidate();
+	LastVideoSampleProcessedTimeRange = TRange<FMediaTimeStamp>::Empty();
+	LastBlockOnRange = TRange<FTimespan>::Empty();
+	OnBlockSeqIndex = 0;
 
 	NextEstVideoTimeAtFrameStart.Invalidate();
 }
@@ -1510,8 +1749,6 @@ void FMediaPlayerFacade::SelectDefaultTracks()
 
 	IMediaTracks& Tracks = Player->GetTracks();
 
-	// @todo gmp: consider locale when selecting default media tracks
-
 	FMediaPlayerTrackOptions TrackOptions;
 	if (ActivePlayerOptions.IsSet())
 	{
@@ -1526,6 +1763,12 @@ void FMediaPlayerFacade::SelectDefaultTracks()
 }
 
 
+float FMediaPlayerFacade::GetUnpausedRate() const
+{
+	return (CurrentRate == 0.0f) ? LastRate : CurrentRate;
+}
+
+
 /* IMediaClockSink interface
 *****************************************************************************/
 
@@ -1535,48 +1778,31 @@ void FMediaPlayerFacade::TickInput(FTimespan DeltaTime, FTimespan Timecode)
 
 	if (Player.IsValid())
 	{
-		Player->TickInput(DeltaTime, Timecode);
-
 		// Update flag reflecting presence of audio in the current stream
-		// (doing it just once per gameloop is enough
+		// (doing it just once per gameloop is enough)
 		bHaveActiveAudio = HaveAudioPlayback();
 
-		// get current play rate
-		float Rate = Player->GetControls().GetRate();
-
-		if (Rate == 0.0f)
-		{
-			Rate = LastRate;
-		}
-		else
-		{
-			LastRate = Rate;
-		}
+		Player->TickInput(DeltaTime, Timecode);
 
 		if (Player->GetPlayerFeatureFlag(IMediaPlayer::EFeatureFlag::UsePlaybackTimingV2))
 		{
 			//
 			// New timing control (handled before any engine world, object etc. updates; so "all frame" (almost) see the state produced here)
 			//
+
 			if (bIsSinkFlushPending)
 			{
 				bIsSinkFlushPending = false;
 				FlushSinks();
 			}
 
-			// Do we have a current timestamp estimation?
-			if (!bHaveActiveAudio && !NextEstVideoTimeAtFrameStart.IsValid())
-			{
-				// Not, yet. We need to attempt to get the next video sample's timestamp to get going...
-				FMediaTimeStamp VideoTimeStamp;
-				if (Player->GetSamples().PeekVideoSampleTime(VideoTimeStamp))
-				{
-					NextEstVideoTimeAtFrameStart = FMediaTimeStampSample(VideoTimeStamp, FPlatformTime::Seconds());
-				}
-			}
+			//
+			// Setup timing for sample processing
+			//
+			PreSampleProcessingTimeHandling();
 
 			TRange<FMediaTimeStamp> TimeRange;
-			if (!GetCurrentPlaybackTimeRange(TimeRange, Rate, DeltaTime, false))
+			if (!GetCurrentPlaybackTimeRange(TimeRange, CurrentRate, DeltaTime, false))
 			{
 				return;
 			}
@@ -1586,29 +1812,43 @@ void FMediaPlayerFacade::TickInput(FTimespan DeltaTime, FTimespan Timecode)
 			//
 			// Process samples in range
 			//
-
 			IMediaSamples& Samples = Player->GetSamples();
-			ProcessCaptionSamples(Samples, TimeRange);
-			ProcessSubtitleSamples(Samples, TimeRange);
-			ProcessVideoSamples(Samples, TimeRange);
 
-			SET_DWORD_STAT(STAT_MediaUtils_FacadeNumVideoSamples, Samples.NumVideoSamples());
-
-			// Move video frame start estimate forward if we have no audio timing to guide us
-			if (!bHaveActiveAudio && NextEstVideoTimeAtFrameStart.IsValid())
+			double BlockingStart = FPlatformTime::Seconds();
+			while (1)
 			{
-				if (Player->GetPlayerFeatureFlag(IMediaPlayer::EFeatureFlag::UseRealtimeWithVideoOnly))
+				ProcessCaptionSamples(Samples, TimeRange);
+				ProcessSubtitleSamples(Samples, TimeRange);
+
+				if (ProcessVideoSamples(Samples, TimeRange))
 				{
-					double NewBaseTime = FPlatformTime::Seconds();
-					NextEstVideoTimeAtFrameStart.TimeStamp.Time += (NewBaseTime - NextEstVideoTimeAtFrameStart.SampledAtTime) * Rate;
-					NextEstVideoTimeAtFrameStart.SampledAtTime = NewBaseTime;
+					break;
 				}
-				else
+
+				// No sample. Should we block for one?
+				if (!BlockOnFetch())
 				{
-					NextEstVideoTimeAtFrameStart.TimeStamp.Time += DeltaTime * Rate;
+					// No... continue...
+					break;
+				}
+
+				if ((FPlatformTime::Seconds() - BlockingStart) > MEDIAUTILS_MAX_BLOCKONFETCH_SECONDS)
+				{
+					UE_LOG(LogMediaUtils, Error, TEXT("Blocking media playback timed out. Disabling it for this playback session."));
+					BlockOnRangeDisabled = true;
+					break;
 				}
+
+				FPlatformProcess::Sleep(0.0f);
 			}
 
+			SET_DWORD_STAT(STAT_MediaUtils_FacadeNumVideoSamples, Samples.NumVideoSamples());
+
+			//
+			// Advance timing etc.
+			//
+			PostSampleProcessingTimeHandling(DeltaTime);
+
 			if (bHaveActiveAudio)
 			{
 				// Keep currently last processed audio sample timestamp available for all frame (to provide consistent info)
@@ -1647,14 +1887,6 @@ void FMediaPlayerFacade::TickFetch(FTimespan DeltaTime, FTimespan Timecode)
 		return;
 	}
 
-	// get current play rate
-	float Rate = Player->GetControls().GetRate();
-
-	if (Rate == 0.0f)
-	{
-		Rate = LastRate;
-	}
-
 	if (!Player->GetPlayerFeatureFlag(IMediaPlayer::EFeatureFlag::UsePlaybackTimingV2))
 	{
 		//
@@ -1666,6 +1898,9 @@ void FMediaPlayerFacade::TickFetch(FTimespan DeltaTime, FTimespan Timecode)
 
 		SET_FLOAT_STAT(STAT_MediaUtils_FacadeTime, CurrentTime.GetTotalSeconds());
 
+		// get current play rate
+		float Rate = GetUnpausedRate();
+
 		if (Rate > 0.0f)
 		{
 			TimeRange = TRange<FTimespan>::AtMost(CurrentTime);
@@ -1687,9 +1922,9 @@ void FMediaPlayerFacade::TickFetch(FTimespan DeltaTime, FTimespan Timecode)
 
 		while (true)
 		{
-			ProcessCaptionSamples(Samples, TimeRange);
-			ProcessSubtitleSamples(Samples, TimeRange);
-			ProcessVideoSamples(Samples, TimeRange);
+			ProcessCaptionSamplesV1(Samples, TimeRange);
+			ProcessSubtitleSamplesV1(Samples, TimeRange);
+			ProcessVideoSamplesV1(Samples, TimeRange);
 
 			if (!BlockOnFetch())
 			{
@@ -1702,7 +1937,7 @@ void FMediaPlayerFacade::TickFetch(FTimespan DeltaTime, FTimespan Timecode)
 				{
 					UE_LOG(LogMediaUtils, Verbose, TEXT("PlayerFacade %p: Aborted block on fetch %s after %i seconds"),
 						this,
-						*BlockOnTime.ToString(TEXT("%h:%m:%s.%t")),
+						*BlockOnRange.GetLowerBoundValue().Time.ToString(TEXT("%h:%m:%s.%t")),
 						MEDIAUTILS_MAX_BLOCKONFETCH_SECONDS
 					);
 
@@ -1711,7 +1946,7 @@ void FMediaPlayerFacade::TickFetch(FTimespan DeltaTime, FTimespan Timecode)
 			}
 			else
 			{
-				UE_LOG(LogMediaUtils, VeryVerbose, TEXT("PlayerFacade %p: Blocking on fetch %s"), this, *BlockOnTime.ToString(TEXT("%h:%m:%s.%t")));
+				UE_LOG(LogMediaUtils, VeryVerbose, TEXT("PlayerFacade %p: Blocking on fetch %s"), this, *BlockOnRange.GetLowerBoundValue().Time.ToString(TEXT("%h:%m:%s.%t")));
 
 				Blocked = true;
 				BlockedTime = FDateTime::UtcNow();
@@ -1732,8 +1967,7 @@ void FMediaPlayerFacade::TickOutput(FTimespan DeltaTime, FTimespan /*Timecode*/)
 		return;
 	}
 
-	IMediaControls& Controls = Player->GetControls();
-	Cache->Tick(DeltaTime, Controls.GetRate(), GetTime());
+	Cache->Tick(DeltaTime, CurrentRate, GetTime());
 }
 
 
@@ -1744,14 +1978,15 @@ void FMediaPlayerFacade::TickTickable()
 {
 	SCOPE_CYCLE_COUNTER(STAT_MediaUtils_FacadeTickTickable);
 
-	if (LastRate == 0.0f)
+	FScopeLock Lock(&CriticalSection);
+
+	if (!Player.IsValid())
 	{
 		return;
 	}
 
-	FScopeLock Lock(&CriticalSection);
-
-	if (!Player.IsValid())
+	float Rate = GetUnpausedRate();
+	if (Rate == 0.0f)
 	{
 		return;
 	}
@@ -1771,7 +2006,7 @@ void FMediaPlayerFacade::TickTickable()
 
 	bool bUseV2Timing = Player->GetPlayerFeatureFlag(IMediaPlayer::EFeatureFlag::UsePlaybackTimingV2);
 
-	if (LastRate > 0.0f)
+	if (Rate > 0.0f)
 	{
 		if (!bUseV2Timing) // we leave range open - sends all the player has
 		{
@@ -1798,6 +2033,104 @@ void FMediaPlayerFacade::TickTickable()
 }
 
 
+void FMediaPlayerFacade::PreSampleProcessingTimeHandling()
+{
+	check(Player.IsValid());
+
+	// No Audio clock?
+	if (!bHaveActiveAudio)
+	{
+		// No external clock? (blocking)
+		if (BlockOnRange.IsEmpty())
+		{
+			// Do we have a current timestamp estimation?
+			if (!NextEstVideoTimeAtFrameStart.IsValid())
+			{
+				// Not, yet. We need to attempt to get the next video sample's timestamp to get going...
+				FMediaTimeStamp VideoTimeStamp;
+				if (Player->GetSamples().PeekVideoSampleTime(VideoTimeStamp))
+				{
+					NextEstVideoTimeAtFrameStart = FMediaTimeStampSample(VideoTimeStamp, FPlatformTime::Seconds());
+				}
+			}
+		}
+	}
+}
+
+
+void FMediaPlayerFacade::PostSampleProcessingTimeHandling(FTimespan DeltaTime)
+{
+	check(Player.IsValid());
+
+	float Rate = CurrentRate;
+
+	// No Audio clock?
+	if (!bHaveActiveAudio)
+	{
+		// No external clock? (blocking)
+		if (BlockOnRange.IsEmpty())
+		{
+			// Move video frame start estimate forward
+			// (the initial NextEstVideoTimeAtFrameStart will never be valid if no video is present)
+			if (!bHaveActiveAudio && NextEstVideoTimeAtFrameStart.IsValid())
+			{
+				if (Player->GetPlayerFeatureFlag(IMediaPlayer::EFeatureFlag::UseRealtimeWithVideoOnly))
+				{
+					double NewBaseTime = FPlatformTime::Seconds();
+					NextEstVideoTimeAtFrameStart.TimeStamp.Time += (NewBaseTime - NextEstVideoTimeAtFrameStart.SampledAtTime) * Rate;
+					NextEstVideoTimeAtFrameStart.SampledAtTime = NewBaseTime;
+				}
+				else
+				{
+					NextEstVideoTimeAtFrameStart.TimeStamp.Time += DeltaTime * Rate;
+				}
+
+				// note: infinite duration (e.g. live playback - or players not yet supporting sequence indices on loops, when looping is enabled)
+				// -> no need for special handling as FTimespan::MaxValue() is expected to be returned to signify this, which is quite "infinite" in practical terms
+				FTimespan Duration = Player->GetControls().GetDuration();
+
+				if (Player->GetControls().IsLooping())
+				{
+					if (Rate >= 0.0f)
+					{
+						if (NextEstVideoTimeAtFrameStart.TimeStamp.Time >= Duration)
+						{
+							NextEstVideoTimeAtFrameStart.TimeStamp.Time -= Duration;
+							++NextEstVideoTimeAtFrameStart.TimeStamp.SequenceIndex;
+						}
+					}
+					else
+					{
+						if (NextEstVideoTimeAtFrameStart.TimeStamp.Time < FTimespan::Zero())
+						{
+							NextEstVideoTimeAtFrameStart.TimeStamp.Time += Duration;
+							--NextEstVideoTimeAtFrameStart.TimeStamp.SequenceIndex;
+						}
+					}
+				}
+				else
+				{
+					if (Rate >= 0.0f)
+					{
+						if (NextEstVideoTimeAtFrameStart.TimeStamp.Time >= Duration)
+						{
+							NextEstVideoTimeAtFrameStart.TimeStamp.Time = Duration - FTimespan::FromSeconds(0.0001);
+						}
+					}
+					else
+					{
+						if (NextEstVideoTimeAtFrameStart.TimeStamp.Time < FTimespan::Zero())
+						{
+							NextEstVideoTimeAtFrameStart.TimeStamp.Time = FTimespan::Zero();
+						}
+					}
+				}
+			}
+		}
+	}
+}
+
+
 bool FMediaPlayerFacade::GetCurrentPlaybackTimeRange(TRange<FMediaTimeStamp> & TimeRange, float Rate, FTimespan DeltaTime, bool bDoNotUseFrameStartReference) const
 {
 	check(Player->GetPlayerFeatureFlag(IMediaPlayer::EFeatureFlag::UsePlaybackTimingV2));
@@ -1850,31 +2183,87 @@ bool FMediaPlayerFacade::GetCurrentPlaybackTimeRange(TRange<FMediaTimeStamp> & T
 			EstAudioTimeAtFrameStart = AudioTime.TimeStamp + FTimespan::FromSeconds((Now - AudioTime.SampledAtTime) * Rate);
 		}
 
+		// Are we paused?
+		if (Rate == 0.0f)
+		{
+			// Yes. We need to fetch a frame for the current display frame - once. Asking over and over until we get one...
+			if (LastVideoSampleProcessedTimeRange.IsEmpty())
+			{
+				// We simply fake the rate to the last non-zero or 1.0 to fetch a frame fitting the time frame representing the whole current frame
+				Rate = (LastRate == 0.0f) ? 1.0f : LastRate;
+			}
+		}
+
 		TimeRange = (Rate >= 0.0f) ? TRange<FMediaTimeStamp>(EstAudioTimeAtFrameStart, EstAudioTimeAtFrameStart + DeltaTime * Rate)
-								   : TRange<FMediaTimeStamp>(EstAudioTimeAtFrameStart + DeltaTime * Rate, EstAudioTimeAtFrameStart);
+								   : TRange<FMediaTimeStamp>(EstAudioTimeAtFrameStart + DeltaTime * (1.0f + Rate), EstAudioTimeAtFrameStart + DeltaTime);
 	}
 	else
 	{
 		//
 		// No Audio (no data and/or no sink)
 		//
-
-		// Do we now have a current timestamp estimation?
-		if (!NextEstVideoTimeAtFrameStart.IsValid())
+		if (BlockOnRange.IsEmpty())
 		{
-			// No timing info available, no timerange available, no samples to process
-			return false;
+			//
+			// Internal clock (DT based)
+			//
+
+			// Do we now have a current timestamp estimation?
+			if (!NextEstVideoTimeAtFrameStart.IsValid())
+			{
+				// No timing info available, no timerange available, no samples to process
+				return false;
+			}
+			else
+			{
+				// Yes. Setup current time range & advance time estimation...
+
+				// Are we paused?
+				if (Rate == 0.0f)
+				{
+					// Yes. We need to fetch a frame for the current display frame - once. Asking over and over until we get one...
+					if (LastVideoSampleProcessedTimeRange.IsEmpty())
+					{
+						// We simply fake the rate to the last non-zero or 1.0 to fetch a frame fitting the time frame representing the whole current frame
+						Rate = (LastRate == 0.0f) ? 1.0f : LastRate;
+					}
+				}
+
+				TimeRange = (Rate >= 0.0f) ? TRange<FMediaTimeStamp>(NextEstVideoTimeAtFrameStart.TimeStamp, NextEstVideoTimeAtFrameStart.TimeStamp + DeltaTime * Rate)
+										   : TRange<FMediaTimeStamp>(NextEstVideoTimeAtFrameStart.TimeStamp + DeltaTime * (1.0f + Rate), NextEstVideoTimeAtFrameStart.TimeStamp + DeltaTime);
+
+				// If we are looping we check to prepare proper ranges should we wrap around either end of the media...
+				// (we do not clamp in the non-looping case as the rest of the code should deal with that fine)
+				if (Player->GetControls().IsLooping())
+				{
+					const FTimespan Duration = Player->GetControls().GetDuration();
+					FTimespan WrappedStart = WrappedModulo(TimeRange.GetLowerBoundValue().Time, Duration);
+					FTimespan WrappedEnd = WrappedModulo(TimeRange.GetUpperBoundValue().Time, Duration);
+					if (WrappedStart > WrappedEnd)
+					{
+						if (Rate >= 0.0)
+						{
+							TimeRange.SetUpperBoundValue(FMediaTimeStamp(WrappedEnd, TimeRange.GetUpperBoundValue().SequenceIndex + 1));
+						}
+						else
+						{
+							TimeRange.SetLowerBoundValue(FMediaTimeStamp(WrappedStart, TimeRange.GetLowerBoundValue().SequenceIndex - 1));
+						}
+					}
+				}
+			}
 		}
 		else
 		{
-			// Yes. Setup current time range & advance time estimation...
-
-			TimeRange = (Rate >= 0.0f) ? TRange<FMediaTimeStamp>(NextEstVideoTimeAtFrameStart.TimeStamp, NextEstVideoTimeAtFrameStart.TimeStamp + DeltaTime * Rate)
-									   : TRange<FMediaTimeStamp>(NextEstVideoTimeAtFrameStart.TimeStamp + DeltaTime * Rate, NextEstVideoTimeAtFrameStart.TimeStamp);
+			//
+			// External clock delivers time-range
+			// (for now we just use the blocking time range as this clock type is solely used in that case)
+			//
+			TimeRange = BlockOnRange;
 		}
 	}
 
-	return true;
+	return !TimeRange.IsEmpty();
 }
 
 
@@ -1924,7 +2313,7 @@ void FMediaPlayerFacade::ProcessAudioSamples(IMediaSamples& Samples, TRange<FTim
 				// (rough estimate, as this is off-gamethread; but better than throwing things out with no throttling at all)
 				{
 				FScopeLock Lock(&LastTimeValuesCS);
-				TimeRange.SetUpperBound(TRangeBound<FTimespan>(LastVideoSampleProcessedTime.TimeStamp.Time));
+				TimeRange.SetUpperBound(TRangeBound<FTimespan>(CurrentFrameVideoTimeStamp.Time));
 				}
 				while (Samples.FetchAudio(TimeRange, Sample))
 					;
@@ -1967,7 +2356,7 @@ void FMediaPlayerFacade::ProcessAudioSamples(IMediaSamples& Samples, TRange<FTim
 }
 
 
-void FMediaPlayerFacade::ProcessCaptionSamples(IMediaSamples& Samples, TRange<FTimespan> TimeRange)
+void FMediaPlayerFacade::ProcessCaptionSamplesV1(IMediaSamples& Samples, TRange<FTimespan> TimeRange)
 {
 	TSharedPtr<IMediaOverlaySample, ESPMode::ThreadSafe> Sample;
 
@@ -1999,7 +2388,7 @@ void FMediaPlayerFacade::ProcessMetadataSamples(IMediaSamples& Samples, TRange<F
 }
 
 
-void FMediaPlayerFacade::ProcessSubtitleSamples(IMediaSamples& Samples, TRange<FTimespan> TimeRange)
+void FMediaPlayerFacade::ProcessSubtitleSamplesV1(IMediaSamples& Samples, TRange<FTimespan> TimeRange)
 {
 	TSharedPtr<IMediaOverlaySample, ESPMode::ThreadSafe> Sample;
 
@@ -2015,7 +2404,7 @@ void FMediaPlayerFacade::ProcessSubtitleSamples(IMediaSamples& Samples, TRange<F
 }
 
 
-void FMediaPlayerFacade::ProcessVideoSamples(IMediaSamples& Samples, TRange<FTimespan> TimeRange)
+void FMediaPlayerFacade::ProcessVideoSamplesV1(IMediaSamples& Samples, TRange<FTimespan> TimeRange)
 {
 	// Let the player do some processing if needed.
 	if (Player.IsValid())
@@ -2037,15 +2426,14 @@ void FMediaPlayerFacade::ProcessVideoSamples(IMediaSamples& Samples, TRange<FTim
 
 		{
 		FScopeLock Lock(&LastTimeValuesCS);
-		LastVideoSampleProcessedTime.TimeStamp = FMediaTimeStamp(Sample->GetTime());
-		LastVideoSampleProcessedTime.SampledAtTime = FPlatformTime::Seconds();
+		CurrentFrameVideoTimeStamp = Sample->GetTime();
 		}
 
 		UE_LOG(LogMediaUtils, VeryVerbose, TEXT("PlayerFacade %p: Fetched video sample %s"), this, *Sample->GetTime().Time.ToString(TEXT("%h:%m:%s.%t")));
 
 		if (VideoSampleSinks.Enqueue(Sample.ToSharedRef(), FMediaPlayerQueueDepths::MaxVideoSinkDepth))
 		{
-			if (GetRate() >= 0.0f)
+			if (CurrentRate >= 0.0f)
 			{
 				NextVideoSampleTime = Sample->GetTime().Time + Sample->GetDuration();
 				UE_LOG(LogMediaUtils, VeryVerbose, TEXT("PlayerFacade %p: Next video sample time %s"), this, *NextVideoSampleTime.ToString(TEXT("%h:%m:%s.%t")));
@@ -2061,7 +2449,7 @@ void FMediaPlayerFacade::ProcessVideoSamples(IMediaSamples& Samples, TRange<FTim
 }
 
 
-void FMediaPlayerFacade::ProcessVideoSamples(IMediaSamples& Samples, const TRange<FMediaTimeStamp> & TimeRange)
+bool FMediaPlayerFacade::ProcessVideoSamples(IMediaSamples& Samples, const TRange<FMediaTimeStamp> & TimeRange)
 {
 	// Let the player do some processing if needed.
 	if (Player.IsValid())
@@ -2082,24 +2470,12 @@ void FMediaPlayerFacade::ProcessVideoSamples(IMediaSamples& Samples, const TRang
 		//
 		// Normal playback with timing control provided by MediaFramework
 		//
-		switch (Samples.FetchBestVideoSampleForTimeRange(TimeRange, Sample, Player->GetControls().GetRate() < 0.0f))
+		const bool bReverse = (GetUnpausedRate() < 0.0f);
+
+		switch (Samples.FetchBestVideoSampleForTimeRange(TimeRange, Sample, bReverse))
 		{
 			case IMediaSamples::EFetchBestSampleResult::Ok:
-			{
-				// Enqueue the sample to render
-				// (we use a queue to stay compatible with existing structure and older sinks - new sinks will read this single entry right away on the gamethread
-				//  and pass it along to rendering outside the queue)
-				bool bOk = VideoSampleSinks.Enqueue(Sample.ToSharedRef(), FMediaPlayerQueueDepths::MaxVideoSinkDepth);
-				check(bOk);
-
-				FScopeLock Lock(&LastTimeValuesCS);
-				LastVideoSampleProcessedTime.TimeStamp = FMediaTimeStamp(Sample->GetTime());
-				LastVideoSampleProcessedTime.SampledAtTime = FPlatformTime::Seconds();
-
-				//uint32 Frames = uint32(LastVideoSampleProcessedTime.TimeStamp.Time.GetTotalMilliseconds() / Sample->GetDuration().GetTotalMilliseconds() + 0.5);
-				//UE_LOG(LogMediaUtils, Display, TEXT("TIME: F=%d"), Frames % 30);
 				break;
-			}
 
 			case IMediaSamples::EFetchBestSampleResult::NoSample:
 				break;
@@ -2110,7 +2486,6 @@ void FMediaPlayerFacade::ProcessVideoSamples(IMediaSamples& Samples, const TRang
 				// Fallback for players supporting V2 timing, but do not supply FetchBestVideoSampleForTimeRange() due to some
 				// custom implementation of IMediaSamples (here to ease adoption of the new timing code - eventually should go away)
 				//
-				const bool bReverse = (Player->GetControls().GetRate() < 0.0f);
 
 				// Find newest sample that satisfies the time range
 				// (the FetchXYZ() code does not work well with a lower range limit at all - we ask for a "up to" type range instead
@@ -2122,15 +2497,11 @@ void FMediaPlayerFacade::ProcessVideoSamples(IMediaSamples& Samples, const TRang
 					((!bReverse && ((Sample->GetTime() + Sample->GetDuration()) > TimeRange.GetLowerBoundValue())) ||
 					(bReverse && ((Sample->GetTime() - Sample->GetDuration()) < TimeRange.GetLowerBoundValue()))))
 				{
-					// Enqueue the sample to render
-					// (we use a queue to stay compatible with existing structure and older sinks - new sinks will read this single entry right away on the gamethread
-					//  and pass it along to rendering outside the queue)
-					bool bOk = VideoSampleSinks.Enqueue(Sample.ToSharedRef(), FMediaPlayerQueueDepths::MaxVideoSinkDepth);
-					check(bOk);
-
-					FScopeLock Lock(&LastTimeValuesCS);
-					LastVideoSampleProcessedTime.TimeStamp = Sample->GetTime();
-					LastVideoSampleProcessedTime.SampledAtTime = FPlatformTime::Seconds();
+					// Sample is good - nothing more to do here
+				}
+				else
+				{
+					Sample.Reset();
 				}
 				break;
 			}
@@ -2144,7 +2515,18 @@ void FMediaPlayerFacade::ProcessVideoSamples(IMediaSamples& Samples, const TRang
 		TRange<FMediaTimeStamp> TempRange; // fully open range
 		while (Samples.FetchVideo(TempRange, Sample))
 			;
-		if (Sample.IsValid())
+	}
+
+	// Any sample?
+	if (Sample.IsValid())
+	{
+		// Yes. If we are in blocking playback mode we need to make sure that the sample is really in the range we asked for and block on...
+		// (same players might return an older sample as stop-gap measure if nothing can be found in the current range)
+
+		TRange<FMediaTimeStamp> SampleTimeRange(Sample->GetTime(), Sample->GetTime() + Sample->GetDuration());
+
+		// Is it what we want?
+		if (BlockOnRange.IsEmpty() || BlockOnRange.Overlaps(SampleTimeRange))
 		{
 			// Enqueue the sample to render
 			// (we use a queue to stay compatible with existing structure and older sinks - new sinks will read this single entry right away on the gamethread
@@ -2153,11 +2535,13 @@ void FMediaPlayerFacade::ProcessVideoSamples(IMediaSamples& Samples, const TRang
 			check(bOk);
 
 			FScopeLock Lock(&LastTimeValuesCS);
-			LastVideoSampleProcessedTime.TimeStamp = Sample->GetTime();
-			LastVideoSampleProcessedTime.SampledAtTime = FPlatformTime::Seconds();
+			CurrentFrameVideoTimeStamp = SampleTimeRange.GetLowerBoundValue();
+			LastVideoSampleProcessedTimeRange = SampleTimeRange;
+
+			return true;
 		}
 	}
-
+	return false;
 }
 
 
@@ -2213,12 +2597,21 @@ void FMediaPlayerFacade::ReceiveMediaEvent(EMediaEvent Event)
 				TSharedPtr<IMediaPlayer, ESPMode::ThreadSafe> CurrentPlayer = Player;
 
 				if (!CurrentPlayer.IsValid())
+				{
 					return;
+				}
 
 				// We only support this for V2 timing players
 				check(CurrentPlayer->GetPlayerFeatureFlag(IMediaPlayer::EFeatureFlag::UsePlaybackTimingV2));
 
-				float Rate = CurrentPlayer->GetControls().GetRate();
+				// Only do this if we do not block on time ranges
+				if (!BlockOnRange.IsEmpty())
+				{
+					// We do not purge as we do not need max perf, but max reliability to actually get certain frames
+					return;
+				}
+
+				float Rate = CurrentRate;
 				if (Rate == 0.0f)
 				{
 					return;
diff --git a/Engine/Source/Runtime/MediaUtils/Public/MediaPlayerFacade.h b/Engine/Source/Runtime/MediaUtils/Public/MediaPlayerFacade.h
--- a/Engine/Source/Runtime/MediaUtils/Public/MediaPlayerFacade.h
+++ b/Engine/Source/Runtime/MediaUtils/Public/MediaPlayerFacade.h
@@ -553,9 +553,20 @@ public:
 	 *
 	 * @param Time The time to block on, or FTimespan::MinValue to disable.
 	 * @see TickFetch
+	 * @note Deprecated: Use SetBlockOnTimeRange instead
 	 */
 	void SetBlockOnTime(const FTimespan& Time);
 
+	/**
+	 * Set the time range on which to block.
+	 *
+	 * If set, this player will block in TickFetch until the video sample
+	 * for the specified time are actually available.
+	 *
+	 * @param TimeRange The time range to block on, use empty range to disable
+	 */
+	void SetBlockOnTimeRange(const TRange<FTimespan>& TimeRange);
+
 	/**
 	 * Set sample caching options.
 	 *
@@ -754,21 +765,22 @@ protected:
 protected:
 	bool HaveAudioPlayback() const;
 	bool HaveVideoPlayback() const;
+	float GetUnpausedRate() const;
 
 	/** Fetch audio samples from the player and forward them to the registered sinks. */
 	void ProcessAudioSamples(IMediaSamples& Samples, TRange<FTimespan> TimeRange);
 
-	/** Fetch audio samples from the player and forward them to the registered sinks. */
-	void ProcessCaptionSamples(IMediaSamples& Samples, TRange<FTimespan> TimeRange);
-
 	/** Fetch metadata samples from the player and forward them to the registered sinks. */
 	void ProcessMetadataSamples(IMediaSamples& Samples, TRange<FTimespan> TimeRange);
 
+	/** Fetch audio samples from the player and forward them to the registered sinks. */
+	void ProcessCaptionSamplesV1(IMediaSamples& Samples, TRange<FTimespan> TimeRange);
+
 	/** Fetch subtitle samples from the player and forward them to the registered sinks. */
-	void ProcessSubtitleSamples(IMediaSamples& Samples, TRange<FTimespan> TimeRange);
+	void ProcessSubtitleSamplesV1(IMediaSamples& Samples, TRange<FTimespan> TimeRange);
 
 	/** Fetch video samples from the player and forward them to the registered sinks. */
-	void ProcessVideoSamples(IMediaSamples& Samples, TRange<FTimespan> TimeRange);
+	void ProcessVideoSamplesV1(IMediaSamples& Samples, TRange<FTimespan> TimeRange);
 
 protected:
 
@@ -787,11 +799,13 @@ private:
 	bool NotifyLifetimeManagerDelegate_PlayerDestroyed();
 	bool NotifyLifetimeManagerDelegate_PlayerResourcesReleased(uint32 ResourceFlags);
 
-	void ProcessVideoSamples(IMediaSamples& Samples, const TRange<FMediaTimeStamp> & TimeRange);
+	bool ProcessVideoSamples(IMediaSamples& Samples, const TRange<FMediaTimeStamp>& TimeRange);
 	void ProcessCaptionSamples(IMediaSamples& Samples, TRange<FMediaTimeStamp> TimeRange);
 	void ProcessSubtitleSamples(IMediaSamples& Samples, TRange<FMediaTimeStamp> TimeRange);
 
-	bool GetCurrentPlaybackTimeRange(TRange<FMediaTimeStamp> & TimeRange, float Rate, FTimespan DeltaTime, bool bDoNotUseFrameStartReference) const;
+	void PreSampleProcessingTimeHandling();
+	bool GetCurrentPlaybackTimeRange(TRange<FMediaTimeStamp>& TimeRange, float Rate, FTimespan DeltaTime, bool bDoNotUseFrameStartReference) const;
+	void PostSampleProcessingTimeHandling(FTimespan DeltaTime);
 
 	void DestroyPlayer();
 
@@ -815,8 +829,17 @@ private:
 
 private:
 
-	/** The time to block on sample fetching. */
-	FTimespan BlockOnTime;
+	/** The time range to block on sample fetching. */
+	TRange<FMediaTimeStamp> BlockOnRange;
+
+	/** Last user provided BlockOnRange value */
+	TRange<FTimespan> LastBlockOnRange;
+
+	/** Flag to indicate block on range feature as disabled for the current playback session due to previous timeout */
+	bool BlockOnRangeDisabled;
+
+	/** Sequence index used during blocked playback processing */
+	int64 OnBlockSeqIndex;
 
 	/** Media sample cache. */
 	FMediaSampleCache* Cache;
@@ -830,6 +853,9 @@ private:
 	/** The last used non-zero play rate (zero if playback never started). */
 	float LastRate;
 
+	/** The last rate set with the facade (unfiltered by player). */
+	float CurrentRate;
+
 	/** Flag indicating that we have an active audio setup */
 	bool bHaveActiveAudio;
 
@@ -867,16 +893,19 @@ private:
 	/** Time of last audio sample decoded. */
 	FMediaTimeStampSample LastAudioSampleProcessedTime;
 
-	/** Time of last video sample decoded. */
-	FMediaTimeStampSample LastVideoSampleProcessedTime;
+	/** Time/Range of last video sample decoded. */
+	TRange<FMediaTimeStamp> LastVideoSampleProcessedTimeRange;
 
 	/** Timestamp for audio considered "current" for this frame (use ONLY for return to outside code) */
 	FMediaTimeStamp CurrentFrameAudioTimeStamp;
 
+	/** Timestamp for video considered "current" for this frame (stays valid even after a flush, until new data comes in) */
+	FMediaTimeStamp CurrentFrameVideoTimeStamp;
+
 	/** Estimation for next frame's video timestamp (used when no audio present or active in stream) */
 	FMediaTimeStampSample NextEstVideoTimeAtFrameStart;
 
-	/** Set if sinks are to be flushed at the request of the player. **/
+	/** Set if sinks are to be flushed at the request of the player. */
 	TAtomic<bool>	bIsSinkFlushPending;
 
 	/** Latch for error state of the most recently used player to be queried after it may have been closed. **/
