Enabling blocking for media sample delivery even with players no yet completely started up
#rnx

diff --git a/Engine/Source/Runtime/Media/Public/IMediaTimeSource.h b/Engine/Source/Runtime/Media/Public/IMediaTimeSource.h
--- a/Engine/Source/Runtime/Media/Public/IMediaTimeSource.h
+++ b/Engine/Source/Runtime/Media/Public/IMediaTimeSource.h
@@ -45,7 +45,7 @@ public:
 	FMediaTimeStamp operator + (const FTimespan & Other) const { return FMediaTimeStamp(Time + Other, SequenceIndex); }
 	FMediaTimeStamp operator - (const FTimespan & Other) const { return FMediaTimeStamp(Time - Other, SequenceIndex); }
 
-	FMediaTimeStamp operator - (const FMediaTimeStamp & Other) const { return FMediaTimeStamp(Time - Other.Time, MAX_int64); }
+	FMediaTimeStamp operator - (const FMediaTimeStamp& Other) const { return (Other.SequenceIndex < SequenceIndex) ? FMediaTimeStamp(FTimespan::MaxValue(), MAX_int64) : ((Other.SequenceIndex > SequenceIndex) ? FMediaTimeStamp(FTimespan::MinValue(), MAX_int64) : FMediaTimeStamp(Time - Other.Time, MAX_int64)); }
 
 	FMediaTimeStamp& operator += (const FTimespan & Other) { Time += Other; return *this; }
 	FMediaTimeStamp& operator -= (const FTimespan & Other) { Time -= Other; return *this; }
diff --git a/Engine/Source/Runtime/MediaUtils/Private/MediaPlayerFacade.cpp b/Engine/Source/Runtime/MediaUtils/Private/MediaPlayerFacade.cpp
--- a/Engine/Source/Runtime/MediaUtils/Private/MediaPlayerFacade.cpp
+++ b/Engine/Source/Runtime/MediaUtils/Private/MediaPlayerFacade.cpp
@@ -103,6 +103,7 @@ FMediaPlayerFacade::FMediaPlayerFacade()
 	, AudioSampleAvailability(-1)
 {
 	BlockOnRange = TRange<FMediaTimeStamp>::Empty();
+	BlockOnTimeRange = TRange<FTimespan>::Empty();
 	BlockOnRangeDisabled = false;
 
 	MediaModule = FModuleManager::LoadModulePtr<IMediaModule>("Media");
@@ -251,6 +252,7 @@ void FMediaPlayerFacade::Close()
 	}
 
 	BlockOnRange = TRange<FMediaTimeStamp>::Empty();
+	BlockOnTimeRange = TRange<FTimespan>::Empty();
 
 	Cache->Empty();
 	CurrentUrl.Empty();
@@ -1055,16 +1057,23 @@ bool FMediaPlayerFacade::ContinueOpen(IMediaPlayerLifecycleManagerDelegate::ICon
 		return false;
 	}
 
+	{
 	FScopeLock Lock(&LastTimeValuesCS);
+
+	BlockOnRangeDisabled = false;
 	LastBlockOnRange = TRange<FTimespan>::Empty();
 	OnBlockSeqIndex = 0;
-	BlockOnRangeDisabled = false;
 	LastVideoSampleProcessedTimeRange = TRange<FMediaTimeStamp>::Empty();
 	LastAudioSampleProcessedTime.Invalidate();
 	CurrentFrameVideoTimeStamp.Invalidate();
 	CurrentFrameAudioTimeStamp.Invalidate();
 
 	NextEstVideoTimeAtFrameStart.Invalidate();
+	}
+
+	// Setup blocking as needed
+	// (in case it was setup with the facade prior to open)
+	SetBlockOnTimeRangeInternal(BlockOnTimeRange);
 
 	if (bCreateNewPlayer)
 	{
@@ -1156,10 +1165,23 @@ void FMediaPlayerFacade::SetBlockOnTime(const FTimespan& Time)
 }
 
 
-void FMediaPlayerFacade::SetBlockOnTimeRange(const TRange<FTimespan> & TimeRange)
+void FMediaPlayerFacade::SetBlockOnTimeRange(const TRange<FTimespan>& TimeRange)
 {
 #if !MEDIAPLAYERFACADE_DISABLE_BLOCKING
-	if (TimeRange.IsEmpty() || !Player.IsValid() || !Player->GetControls().CanControl(EMediaControl::BlockOnFetch))
+	BlockOnTimeRange = TimeRange;
+	if (Player.IsValid())
+	{
+		SetBlockOnTimeRangeInternal(TimeRange);
+	}
+#endif
+}
+
+
+void FMediaPlayerFacade::SetBlockOnTimeRangeInternal(const TRange<FTimespan> &TimeRange)
+{
+	check(Player.IsValid());
+
+	if (TimeRange.IsEmpty() || !Player->GetControls().CanControl(EMediaControl::BlockOnFetch))
 	{
 		LastBlockOnRange = TRange<FTimespan>::Empty();
 		BlockOnRange = TRange<FMediaTimeStamp>::Empty();
@@ -1242,7 +1264,6 @@ void FMediaPlayerFacade::SetBlockOnTimeRange(const TRange<FTimespan> & TimeRange
 	Player->GetControls().SetBlockingPlaybackHint(!BlockOnRange.IsEmpty());
 
 	LastBlockOnRange = TimeRange;
-#endif
 }
 
 
@@ -1384,17 +1405,6 @@ bool FMediaPlayerFacade::BlockOnFetch() const
 		return false; // no blocking requested / not supported
 	}
 
-	if (IsPreparing())
-	{
-		return true; // block on media opening
-	}
-
-	if (!IsPlaying())
-	{
-		// no blocking if we are not playing (e.g. paused)
-		return false;
-	}
-
 	if (Player->GetPlayerFeatureFlag(IMediaPlayer::EFeatureFlag::UsePlaybackTimingV2))
 	{
 		//
@@ -1450,33 +1460,18 @@ bool FMediaPlayerFacade::BlockOnFetch() const
 			}
 		}
 
-		if (!Player->GetControls().IsLooping())
+		// The next checks make only sense if the player is done preparing...
+		if (!IsPreparing())
 		{
-			// Is the sample outside the media's range?
-			// (note: this assumes the media starts at time ZERO - this will not be the case at all times (e.g. life playback) -- for now we assume a player will flagged blocked playback as invalid in that case!)
-			if (BlockOnRange.GetUpperBoundValue() < FMediaTimeStamp(0) ||  Player->GetControls().GetDuration() <= BlockOnRange.GetLowerBoundValue().Time)
-			{
-				return false;
-			}
-		}
-
-		// Did we pass the range already?
-		// (this would imply a malfunction or lack of support in the player -- but: we need to catch it to avoid stalls)
-		const FMediaTimeStamp CurrentTime = GetTimeStamp();
-		if (Rate >= 0)
-		{
-			if (CurrentTime >= BlockOnRange.GetUpperBoundValue())
-			{
-				UE_LOG(LogMediaUtils, Warning, TEXT("Blocking range for playback is behind player playback position. Skipping blocking; possible malfunction."));
-				return false;
-			}
-		}
-		else
-		{
-			if (CurrentTime < BlockOnRange.GetLowerBoundValue())
+			// Looping off?
+			if (!Player->GetControls().IsLooping())
 			{
-				UE_LOG(LogMediaUtils, Warning, TEXT("Blocking range for playback is behind player playback position. Skipping blocking; possible malfunction."));
-				return false;
+				// Yes. Is the sample outside the media's range?
+				// (note: this assumes the media starts at time ZERO - this will not be the case at all times (e.g. life playback) -- for now we assume a player will flagged blocked playback as invalid in that case!)
+				if (BlockOnRange.GetUpperBoundValue() < FMediaTimeStamp(0) || Player->GetControls().GetDuration() <= BlockOnRange.GetLowerBoundValue().Time)
+				{
+					return false;
+				}
 			}
 		}
 
@@ -1489,6 +1484,17 @@ bool FMediaPlayerFacade::BlockOnFetch() const
 		// V1 blocking logic
 		//
 
+		if (IsPreparing())
+		{
+			return true; // block on media opening
+		}
+
+		if (!IsPlaying())
+		{
+			// no blocking if we are not playing (e.g. paused)
+			return false;
+		}
+
 		if (CurrentRate < 0.0f)
 		{
 			return false; // block only in forward play
@@ -1790,6 +1796,21 @@ void FMediaPlayerFacade::TickInput(FTimespan DeltaTime, FTimespan Timecode)
 			// New timing control (handled before any engine world, object etc. updates; so "all frame" (almost) see the state produced here)
 			//
 
+			// process deferred events
+			// NOTE: if there is no player anymore we execute the remaining queued events in TickFetch (backwards compatibility - should move here once V1 support removed)
+			EMediaEvent Event;
+			while (QueuedEvents.Dequeue(Event))
+			{
+				ProcessEvent(Event);
+			}
+
+			// Handling events may have killed the player. Did it?
+			if (!Player.IsValid())
+			{
+				// If so: nothing more to do!
+				return;
+			}
+
 			if (bIsSinkFlushPending)
 			{
 				bIsSinkFlushPending = false;
@@ -1832,6 +1853,30 @@ void FMediaPlayerFacade::TickInput(FTimespan DeltaTime, FTimespan Timecode)
 					break;
 				}
 
+				// Issue tick call with dummy timing as some players advance some state in the tick, which we wait for
+				Player->TickInput(FTimespan::Zero(), FTimespan::MinValue());
+
+				// Process deferred events & check for events that break the block
+				bool bEventCancelsBlock = false;
+				while (QueuedEvents.Dequeue(Event))
+				{
+					if (Event == EMediaEvent::MediaClosed || Event == EMediaEvent::MediaOpenFailed)
+					{
+						bEventCancelsBlock = true;
+					}
+					ProcessEvent(Event);
+				}
+
+				// We might have lost the player during event handling or an event breaks the block...
+				if (!Player.IsValid() || bEventCancelsBlock)
+				{
+					// Disable blocking feature for now (a new open would reset this)
+					UE_LOG(LogMediaUtils, Warning, TEXT("Blocking media playback closed or failed. Disabling it for this playback session."));
+					BlockOnRangeDisabled = true;
+					break;
+				}
+
+				// Timeout?
 				if ((FPlatformTime::Seconds() - BlockingStart) > MEDIAUTILS_MAX_BLOCKONFETCH_SECONDS)
 				{
 					UE_LOG(LogMediaUtils, Error, TEXT("Blocking media playback timed out. Disabling it for this playback session."));
@@ -1867,12 +1912,7 @@ void FMediaPlayerFacade::TickFetch(FTimespan DeltaTime, FTimespan Timecode)
 {
 	SCOPE_CYCLE_COUNTER(STAT_MediaUtils_FacadeTickFetch);
 
-	// let the player generate samples & process events
-	if (Player.IsValid())
-	{
-		Player->TickFetch(DeltaTime, Timecode);
-	}
-
+	if (!Player.IsValid())
 	{
 		// process deferred events
 		EMediaEvent Event;
@@ -1880,10 +1920,6 @@ void FMediaPlayerFacade::TickFetch(FTimespan DeltaTime, FTimespan Timecode)
 		{
 			ProcessEvent(Event);
 		}
-	}
-
-	if (!Player.IsValid())
-	{
 		return;
 	}
 
@@ -1892,6 +1928,19 @@ void FMediaPlayerFacade::TickFetch(FTimespan DeltaTime, FTimespan Timecode)
 		//
 		// Old timing control
 		//
+
+		// let the player generate samples & process events
+		Player->TickFetch(DeltaTime, Timecode);
+
+		{
+			// process deferred events
+			EMediaEvent Event;
+			while (QueuedEvents.Dequeue(Event))
+			{
+				ProcessEvent(Event);
+			}
+		}
+
 		TRange<FTimespan> TimeRange;
 
 		const FTimespan CurrentTime = GetTime();
@@ -2478,7 +2527,9 @@ bool FMediaPlayerFacade::ProcessVideoSamples(IMediaSamples& Samples, const TRang
 				break;
 
 			case IMediaSamples::EFetchBestSampleResult::NoSample:
+			{
 				break;
+			}
 
 			case IMediaSamples::EFetchBestSampleResult::NotSupported:
 			{
diff --git a/Engine/Source/Runtime/MediaUtils/Private/MediaUtilsPrivate.h b/Engine/Source/Runtime/MediaUtils/Private/MediaUtilsPrivate.h
--- a/Engine/Source/Runtime/MediaUtils/Private/MediaUtilsPrivate.h
+++ b/Engine/Source/Runtime/MediaUtils/Private/MediaUtilsPrivate.h
@@ -6,8 +6,8 @@
 #include "Stats/Stats.h"
 
 
-/** Maximum time that TickFetch will block waiting for samples (in seconds). */
-#define MEDIAUTILS_MAX_BLOCKONFETCH_SECONDS 5
+/** Maximum time that TickInput/Fetch will block waiting for samples (in seconds). */
+#define MEDIAUTILS_MAX_BLOCKONFETCH_SECONDS 10
 
 
 /** Declares a log category for this module. */
diff --git a/Engine/Source/Runtime/MediaUtils/Public/MediaPlayerFacade.h b/Engine/Source/Runtime/MediaUtils/Public/MediaPlayerFacade.h
--- a/Engine/Source/Runtime/MediaUtils/Public/MediaPlayerFacade.h
+++ b/Engine/Source/Runtime/MediaUtils/Public/MediaPlayerFacade.h
@@ -718,6 +718,11 @@ protected:
 	 */
 	bool BlockOnFetch() const;
 
+	/**
+	 * Internal setup for blocking time ranges
+	 */
+	void SetBlockOnTimeRangeInternal(const TRange<FTimespan>& TimeRange);
+
 	/** Flush all media sample sinks. */
 	void FlushSinks();
 
@@ -835,6 +840,9 @@ private:
 	/** Last user provided BlockOnRange value */
 	TRange<FTimespan> LastBlockOnRange;
 
+	/** Facade's copy of the last value passed into SetBlockOnRange(); used in case evaluation must be delayed */
+	TRange<FTimespan> BlockOnTimeRange;
+
 	/** Flag to indicate block on range feature as disabled for the current playback session due to previous timeout */
 	bool BlockOnRangeDisabled;
 
