Adding feature flag to Media Player Facade to allow for suport of players that handle flushes on seeks internally

diff --git a/Engine/Source/Runtime/Media/Public/IMediaPlayer.h b/Engine/Source/Runtime/Media/Public/IMediaPlayer.h
--- a/Engine/Source/Runtime/Media/Public/IMediaPlayer.h
+++ b/Engine/Source/Runtime/Media/Public/IMediaPlayer.h
@@ -306,10 +306,11 @@ public:
 	}
 
 	enum class EFeatureFlag {
-		AllowShutdownOnClose = 0,	//!< Allow player to be shutdown right after 'close' event is received from it
-		UsePlaybackTimingV2,		//!< Use v2 playback timing and AV sync
-		UseRealtimeWithVideoOnly,	//!< Use realtime rather then game deltatime to control video playback if no audio is present
-		AlwaysPullNewestVideoFrame,	//!< Mediaframework will not gate video frame output with its own timing, but assumes "ASAP" as output time for every sample
+		AllowShutdownOnClose = 0,		//!< Allow player to be shutdown right after 'close' event is received from it
+		UsePlaybackTimingV2,			//!< Use v2 playback timing and AV sync
+		UseRealtimeWithVideoOnly,		//!< Use realtime rather then game deltatime to control video playback if no audio is present
+		AlwaysPullNewestVideoFrame,		//!< Mediaframework will not gate video frame output with its own timing, but assumes "ASAP" as output time for every sample
+		PlayerUsesInternalFlushOnSeek,	//!< The player implements an internal flush logic on seeks and Mediaframework will not issue an explicit Flush() call to it on seeks
 	};
 	
 	virtual bool GetPlayerFeatureFlag(EFeatureFlag /*flag*/) const
diff --git a/Engine/Source/Runtime/MediaUtils/Private/MediaPlayerFacade.cpp b/Engine/Source/Runtime/MediaUtils/Private/MediaPlayerFacade.cpp
--- a/Engine/Source/Runtime/MediaUtils/Private/MediaPlayerFacade.cpp
+++ b/Engine/Source/Runtime/MediaUtils/Private/MediaPlayerFacade.cpp
@@ -264,7 +264,7 @@ void FMediaPlayerFacade::Close()
 	bIsSinkFlushPending = false;
 	bDidRecentPlayerHaveError = false;
 
-	FlushSinks();
+	Flush();
 }
 
 
@@ -1122,7 +1122,7 @@ bool FMediaPlayerFacade::Seek(const FTimespan& Time)
 
 	if (Player.IsValid() && Player->FlushOnSeekStarted())
 	{
-		FlushSinks();
+		Flush(Player->GetPlayerFeatureFlag(IMediaPlayer::EFeatureFlag::PlayerUsesInternalFlushOnSeek));
 	}
 
 	return true;
@@ -1136,7 +1136,7 @@ bool FMediaPlayerFacade::SelectTrack(EMediaTrackType TrackType, int32 TrackIndex
 		return false;
 	}
 
-	FlushSinks();
+	Flush();
 
 	return true;
 }
@@ -1316,7 +1316,7 @@ bool FMediaPlayerFacade::SetRate(float Rate)
 
 	if ((LastRate * Rate) < 0.0f)
 	{
-		FlushSinks(); // direction change
+		Flush(); // direction change
 	}
 	else
 	{
@@ -1514,7 +1514,7 @@ bool FMediaPlayerFacade::BlockOnFetch() const
 }
 
 
-void FMediaPlayerFacade::FlushSinks()
+void FMediaPlayerFacade::Flush(bool bExcludePlayer)
 {
 	UE_LOG(LogMediaUtils, Verbose, TEXT("PlayerFacade %p: Flushing sinks"), this);
 
@@ -1526,7 +1526,7 @@ void FMediaPlayerFacade::FlushSinks()
 	SubtitleSampleSinks.Flush();
 	VideoSampleSinks.Flush();
 
-	if (Player.IsValid())
+	if (Player.IsValid() && !bExcludePlayer)
 	{
 		Player->GetSamples().FlushSamples();
 	}
@@ -1715,13 +1715,13 @@ void FMediaPlayerFacade::ProcessEvent(EMediaEvent Event)
 	if ((Event == EMediaEvent::PlaybackEndReached) ||
 		(Event == EMediaEvent::TracksChanged))
 	{
-		FlushSinks();
+		Flush();
 	}
 	else if (Event == EMediaEvent::SeekCompleted)
 	{
 		if (!Player.IsValid() || Player->FlushOnSeekCompleted())
 		{
-			FlushSinks();
+			Flush(Player->GetPlayerFeatureFlag(IMediaPlayer::EFeatureFlag::PlayerUsesInternalFlushOnSeek));
 		}
 	}
 	else if (Event == EMediaEvent::MediaClosed)
@@ -1815,7 +1815,7 @@ void FMediaPlayerFacade::TickInput(FTimespan DeltaTime, FTimespan Timecode)
 			if (bIsSinkFlushPending)
 			{
 				bIsSinkFlushPending = false;
-				FlushSinks();
+				Flush();
 			}
 
 			//
diff --git a/Engine/Source/Runtime/MediaUtils/Public/MediaPlayerFacade.h b/Engine/Source/Runtime/MediaUtils/Public/MediaPlayerFacade.h
--- a/Engine/Source/Runtime/MediaUtils/Public/MediaPlayerFacade.h
+++ b/Engine/Source/Runtime/MediaUtils/Public/MediaPlayerFacade.h
@@ -723,8 +723,8 @@ protected:
 	 */
 	void SetBlockOnTimeRangeInternal(const TRange<FTimespan>& TimeRange);
 
-	/** Flush all media sample sinks. */
-	void FlushSinks();
+	/** Flush all media sample sinks & player plugin. */
+	void Flush(bool bExcludePlayer = false);
 
 	/**
 	 * Get details about the specified audio track format.
