MediaFramework: Fixed sequence index not increasing when the last block time (ignoring the sequence index) equals the new block time
#rnx

diff --git a/Engine/Source/Runtime/MediaUtils/Private/MediaPlayerFacade.cpp b/Engine/Source/Runtime/MediaUtils/Private/MediaPlayerFacade.cpp
--- a/Engine/Source/Runtime/MediaUtils/Private/MediaPlayerFacade.cpp
+++ b/Engine/Source/Runtime/MediaUtils/Private/MediaPlayerFacade.cpp
@@ -1229,7 +1229,7 @@ void FMediaPlayerFacade::SetBlockOnTimeRangeInternal(const TRange<FTimespan> &Ti
 			if (Rate >= 0.0f)
 			{
 				FTimespan LastStart = WrappedModulo(LastBlockOnRange.GetLowerBoundValue(), Duration);
-				if (LastStart > Start)
+				if ((LastStart > Start) || ((LastStart == Start) && (LastBlockOnRange.GetLowerBoundValue() < TimeRange.GetLowerBoundValue())))
 				{
 					++OnBlockSeqIndex;
 				}
@@ -1237,7 +1237,7 @@ void FMediaPlayerFacade::SetBlockOnTimeRangeInternal(const TRange<FTimespan> &Ti
 			else
 			{
 				FTimespan LastEnd = WrappedModulo(LastBlockOnRange.GetUpperBoundValue(), Duration);
-				if (LastEnd < End)
+				if ((LastEnd < End) || ((LastEnd == End) && (LastBlockOnRange.GetUpperBoundValue() > TimeRange.GetUpperBoundValue())))
 				{
 					--OnBlockSeqIndex;
 				}
