Refactor to BlockOnRange setup & updates to allow for delayed computations to ensure we got proper information from the player before attempting them

diff --git a/Engine/Source/Runtime/MediaUtils/Private/MediaPlayerFacade.cpp b/Engine/Source/Runtime/MediaUtils/Private/MediaPlayerFacade.cpp
--- a/Engine/Source/Runtime/MediaUtils/Private/MediaPlayerFacade.cpp
+++ b/Engine/Source/Runtime/MediaUtils/Private/MediaPlayerFacade.cpp
@@ -95,6 +95,7 @@ static FTimespan WrappedModulo(FTimespan Time, FTimespan Duration)
 
 FMediaPlayerFacade::FMediaPlayerFacade()
 	: TimeDelay(FTimespan::Zero())
+	, BlockOnRange(this)
 	, Cache(new FMediaSampleCache)
 	, LastRate(0.0f)
 	, CurrentRate(0.0f)
@@ -102,8 +103,6 @@ FMediaPlayerFacade::FMediaPlayerFacade()
 	, VideoSampleAvailability(-1)
 	, AudioSampleAvailability(-1)
 {
-	BlockOnRange = TRange<FMediaTimeStamp>::Empty();
-	BlockOnTimeRange = TRange<FTimespan>::Empty();
 	BlockOnRangeDisabled = false;
 
 	MediaModule = FModuleManager::LoadModulePtr<IMediaModule>("Media");
@@ -251,8 +250,7 @@ void FMediaPlayerFacade::Close()
 		NotifyLifetimeManagerDelegate_PlayerClosed();
 	}
 
-	BlockOnRange = TRange<FMediaTimeStamp>::Empty();
-	BlockOnTimeRange = TRange<FTimespan>::Empty();
+	BlockOnRange.Reset();
 
 	Cache->Empty();
 	CurrentUrl.Empty();
@@ -1061,8 +1059,7 @@ bool FMediaPlayerFacade::ContinueOpen(IMediaPlayerLifecycleManagerDelegate::ICon
 	FScopeLock Lock(&LastTimeValuesCS);
 
 	BlockOnRangeDisabled = false;
-	LastBlockOnRange = TRange<FTimespan>::Empty();
-	OnBlockSeqIndex = 0;
+	BlockOnRange.Flush();
 	LastVideoSampleProcessedTimeRange = TRange<FMediaTimeStamp>::Empty();
 	LastAudioSampleProcessedTime.Invalidate();
 	CurrentFrameVideoTimeStamp.Invalidate();
@@ -1071,11 +1068,6 @@ bool FMediaPlayerFacade::ContinueOpen(IMediaPlayerLifecycleManagerDelegate::ICon
 	NextEstVideoTimeAtFrameStart.Invalidate();
 	}
 
-	// Setup blocking as needed
-	// (in case it was setup with the facade prior to open)
-	SetBlockOnTimeRangeInternal(BlockOnTimeRange);
-
-
 	if (bCreateNewPlayer)
 	{
 		NotifyLifetimeManagerDelegate_PlayerCreated();
@@ -1152,14 +1144,14 @@ void FMediaPlayerFacade::SetBlockOnTime(const FTimespan& Time)
 
 	if (Time == FTimespan::MinValue())
 	{
-		BlockOnRange = TRange<FMediaTimeStamp>::Empty();
-		LastBlockOnRange = TRange<FTimespan>::Empty();
+		BlockOnRange.SetRange(TRange<FTimespan>::Empty());
 		Player->GetControls().SetBlockingPlaybackHint(false);
 	}
 	else
 	{
-		BlockOnRange.Inclusive(Time, Time);
-		LastBlockOnRange.Inclusive(Time, Time);
+		TRange<FTimespan> Range;
+		Range.Inclusive(Time, Time);
+		BlockOnRange.SetRange(Range);
 		Player->GetControls().SetBlockingPlaybackHint(true);
 	}
 #endif
@@ -1169,32 +1161,50 @@ void FMediaPlayerFacade::SetBlockOnTime(const FTimespan& Time)
 void FMediaPlayerFacade::SetBlockOnTimeRange(const TRange<FTimespan>& TimeRange)
 {
 #if !MEDIAPLAYERFACADE_DISABLE_BLOCKING
-	BlockOnTimeRange = TimeRange;
-	if (Player.IsValid())
-	{
-		SetBlockOnTimeRangeInternal(TimeRange);
-	}
+	BlockOnRange.SetRange(TimeRange);
 #endif
 }
 
 
-void FMediaPlayerFacade::SetBlockOnTimeRangeInternal(const TRange<FTimespan> &TimeRange)
+void FMediaPlayerFacade::FBlockOnRange::Flush()
 {
-	check(Player.IsValid());
+	LastBlockOnRange = TRange<FTimespan>::Empty();
+	OnBlockSeqIndex = 0;
+}
+
+
+void FMediaPlayerFacade::FBlockOnRange::SetRange(const TRange<FTimespan>& NewRange)
+{
+	CurrentTimeRange = NewRange;
+	RangeIsDirty = true;
+}
 
-	if (TimeRange.IsEmpty() || !Player->GetControls().CanControl(EMediaControl::BlockOnFetch))
+
+const TRange<FMediaTimeStamp> & FMediaPlayerFacade::FBlockOnRange::GetRange() const
+{
+	TSharedPtr<IMediaPlayer, ESPMode::ThreadSafe> CurrentPlayer(Facade->Player);
+	check(CurrentPlayer.IsValid());
+
+	if (!RangeIsDirty)
+	{
+		return BlockOnRange;
+	}
+
+	RangeIsDirty = false;
+
+	if (CurrentTimeRange.IsEmpty() || !CurrentPlayer->GetControls().CanControl(EMediaControl::BlockOnFetch))
 	{
 		LastBlockOnRange = TRange<FTimespan>::Empty();
 		BlockOnRange = TRange<FMediaTimeStamp>::Empty();
-		Player->GetControls().SetBlockingPlaybackHint(false);
-		return;
+		CurrentPlayer->GetControls().SetBlockingPlaybackHint(false);
+		return BlockOnRange;
 	}
 
-	FTimespan Duration(Player->GetControls().GetDuration());
-	FTimespan Start(TimeRange.GetLowerBoundValue());
-	FTimespan End(TimeRange.GetUpperBoundValue());
+	FTimespan Duration(CurrentPlayer->GetControls().GetDuration());
+	FTimespan Start(CurrentTimeRange.GetLowerBoundValue());
+	FTimespan End(CurrentTimeRange.GetUpperBoundValue());
 
-	if (!Player->GetControls().IsLooping())
+	if (!CurrentPlayer->GetControls().IsLooping())
 	{
 		// We pass in the time range as is on seq-index zero at all times - players have to reject sample output / blocking logic will detect begin outside media range
 		BlockOnRange = TRange<FMediaTimeStamp>(FMediaTimeStamp(Start, 0), FMediaTimeStamp(End, 0));
@@ -1212,16 +1222,18 @@ void FMediaPlayerFacade::SetBlockOnTimeRangeInternal(const TRange<FTimespan> &Ti
 		{
 			// Catch if this is called to early and reset blocking...
 			BlockOnRange = TRange<FMediaTimeStamp>::Empty();
-			Player->GetControls().SetBlockingPlaybackHint(false);
-			return;
+			CurrentPlayer->GetControls().SetBlockingPlaybackHint(false);
+			return BlockOnRange;
 		}
 
+
+		float Rate = Facade->GetUnpausedRate();
+
 		// Modulo on the time to get it into media's range
 		// (assumes zero-start-time)
 		Start = WrappedModulo(Start, Duration);
 		End = WrappedModulo(End, Duration);
 
-		float Rate = GetUnpausedRate();
 
 		// Detect any non-monotonic movement of the range...
 		if (!LastBlockOnRange.IsEmpty())
@@ -1229,7 +1241,7 @@ void FMediaPlayerFacade::SetBlockOnTimeRangeInternal(const TRange<FTimespan> &Ti
 			if (Rate >= 0.0f)
 			{
 				FTimespan LastStart = WrappedModulo(LastBlockOnRange.GetLowerBoundValue(), Duration);
-				if ((LastStart > Start) || ((LastStart == Start) && (LastBlockOnRange.GetLowerBoundValue() < TimeRange.GetLowerBoundValue())))
+				if ((LastStart > Start) || ((LastStart == Start) && (LastBlockOnRange.GetLowerBoundValue() < CurrentTimeRange.GetLowerBoundValue())))
 				{
 					++OnBlockSeqIndex;
 				}
@@ -1237,7 +1249,7 @@ void FMediaPlayerFacade::SetBlockOnTimeRangeInternal(const TRange<FTimespan> &Ti
 			else
 			{
 				FTimespan LastEnd = WrappedModulo(LastBlockOnRange.GetUpperBoundValue(), Duration);
-				if ((LastEnd < End) || ((LastEnd == End) && (LastBlockOnRange.GetUpperBoundValue() > TimeRange.GetUpperBoundValue())))
+				if ((LastEnd < End) || ((LastEnd == End) && (LastBlockOnRange.GetUpperBoundValue() > CurrentTimeRange.GetUpperBoundValue())))
 				{
 					--OnBlockSeqIndex;
 				}
@@ -1262,9 +1274,11 @@ void FMediaPlayerFacade::SetBlockOnTimeRangeInternal(const TRange<FTimespan> &Ti
 		check(!BlockOnRange.IsEmpty());
 	}
 
-	Player->GetControls().SetBlockingPlaybackHint(!BlockOnRange.IsEmpty());
+	CurrentPlayer->GetControls().SetBlockingPlaybackHint(!BlockOnRange.IsEmpty());
 
-	LastBlockOnRange = TimeRange;
+	LastBlockOnRange = CurrentTimeRange;
+
+	return BlockOnRange;
 }
 
 
@@ -1401,7 +1415,9 @@ bool FMediaPlayerFacade::BlockOnFetch() const
 {
 	check(Player.IsValid());
 
-	if (BlockOnRange.IsEmpty() || !Player->GetControls().CanControl(EMediaControl::BlockOnFetch) || BlockOnRangeDisabled)
+	const TRange<FMediaTimeStamp> & BR = BlockOnRange.GetRange();
+
+	if (BR.IsEmpty() || !Player->GetControls().CanControl(EMediaControl::BlockOnFetch) || BlockOnRangeDisabled)
 	{
 		return false; // no blocking requested / not supported
 	}
@@ -1414,7 +1430,7 @@ bool FMediaPlayerFacade::BlockOnFetch() const
 		float Rate = GetUnpausedRate();
 
 		// If the current sample "out there" is actually overlapping with the current block, we might be good with no new sample
-		if (LastVideoSampleProcessedTimeRange.Overlaps(BlockOnRange))
+		if (LastVideoSampleProcessedTimeRange.Overlaps(BR))
 		{
 			// We have no new data (else we would not even call this method), but the last sample we returned is still inside the current range -> good, but...
 			// If the next sample would already cover more of the range than the older one we would like to use that instead -> but it may well be we do not have any data about the sample yet (and would indeed LIKE to block!)
@@ -1451,8 +1467,8 @@ bool FMediaPlayerFacade::BlockOnFetch() const
 			}
 
 			// Compute which one is larger inside the current range...
-			int64 LastSampleCoverage = TRange<FMediaTimeStamp>::Intersection(BlockOnRange, LastSampleTimeRange).Size<FMediaTimeStamp>().Time.GetTicks();
-			int64 NextSampleCoverage = TRange<FMediaTimeStamp>::Intersection(BlockOnRange, NextSampleTimeRange).Size<FMediaTimeStamp>().Time.GetTicks();
+			int64 LastSampleCoverage = TRange<FMediaTimeStamp>::Intersection(BR, LastSampleTimeRange).Size<FMediaTimeStamp>().Time.GetTicks();
+			int64 NextSampleCoverage = TRange<FMediaTimeStamp>::Intersection(BR, NextSampleTimeRange).Size<FMediaTimeStamp>().Time.GetTicks();
 
 			if (LastSampleCoverage > NextSampleCoverage)
 			{
@@ -1469,7 +1485,7 @@ bool FMediaPlayerFacade::BlockOnFetch() const
 			{
 				// Yes. Is the sample outside the media's range?
 				// (note: this assumes the media starts at time ZERO - this will not be the case at all times (e.g. life playback) -- for now we assume a player will flagged blocked playback as invalid in that case!)
-				if (BlockOnRange.GetUpperBoundValue() < FMediaTimeStamp(0) || Player->GetControls().GetDuration() <= BlockOnRange.GetLowerBoundValue().Time)
+				if (BR.GetUpperBoundValue() < FMediaTimeStamp(0) || Player->GetControls().GetDuration() <= BR.GetLowerBoundValue().Time)
 				{
 					return false;
 				}
@@ -1501,7 +1517,7 @@ bool FMediaPlayerFacade::BlockOnFetch() const
 			return false; // block only in forward play
 		}
 
-		const bool VideoReady = (VideoSampleSinks.Num() == 0) || (BlockOnRange.GetUpperBoundValue() < NextVideoSampleTime);
+		const bool VideoReady = (VideoSampleSinks.Num() == 0) || (BR.GetUpperBoundValue() < NextVideoSampleTime);
 
 		if (VideoReady)
 		{
@@ -1535,8 +1551,7 @@ void FMediaPlayerFacade::Flush(bool bExcludePlayer)
 
 	LastAudioRenderedSampleTime.Invalidate();
 	LastVideoSampleProcessedTimeRange = TRange<FMediaTimeStamp>::Empty();
-	LastBlockOnRange = TRange<FTimespan>::Empty();
-	OnBlockSeqIndex = 0;
+	BlockOnRange.Flush();
 
 	NextEstVideoTimeAtFrameStart.Invalidate();
 }
@@ -1987,7 +2002,7 @@ void FMediaPlayerFacade::TickFetch(FTimespan DeltaTime, FTimespan Timecode)
 				{
 					UE_LOG(LogMediaUtils, Verbose, TEXT("PlayerFacade %p: Aborted block on fetch %s after %i seconds"),
 						this,
-						*BlockOnRange.GetLowerBoundValue().Time.ToString(TEXT("%h:%m:%s.%t")),
+						*BlockOnRange.GetRange().GetLowerBoundValue().Time.ToString(TEXT("%h:%m:%s.%t")),
 						MEDIAUTILS_MAX_BLOCKONFETCH_SECONDS
 					);
 
@@ -1996,7 +2011,7 @@ void FMediaPlayerFacade::TickFetch(FTimespan DeltaTime, FTimespan Timecode)
 			}
 			else
 			{
-				UE_LOG(LogMediaUtils, VeryVerbose, TEXT("PlayerFacade %p: Blocking on fetch %s"), this, *BlockOnRange.GetLowerBoundValue().Time.ToString(TEXT("%h:%m:%s.%t")));
+				UE_LOG(LogMediaUtils, VeryVerbose, TEXT("PlayerFacade %p: Blocking on fetch %s"), this, *BlockOnRange.GetRange().GetLowerBoundValue().Time.ToString(TEXT("%h:%m:%s.%t")));
 
 				Blocked = true;
 				BlockedTime = FDateTime::UtcNow();
@@ -2091,7 +2106,7 @@ void FMediaPlayerFacade::PreSampleProcessingTimeHandling()
 	if (!bHaveActiveAudio)
 	{
 		// No external clock? (blocking)
-		if (BlockOnRange.IsEmpty())
+		if (BlockOnRange.GetRange().IsEmpty())
 		{
 			// Do we have a current timestamp estimation?
 			if (!NextEstVideoTimeAtFrameStart.IsValid())
@@ -2118,7 +2133,7 @@ void FMediaPlayerFacade::PostSampleProcessingTimeHandling(FTimespan DeltaTime)
 	if (!bHaveActiveAudio)
 	{
 		// No external clock? (blocking)
-		if (BlockOnRange.IsEmpty())
+		if (BlockOnRange.GetRange().IsEmpty())
 		{
 			// Move video frame start estimate forward
 			// (the initial NextEstVideoTimeAtFrameStart will never be valid if no video is present)
@@ -2252,7 +2267,7 @@ bool FMediaPlayerFacade::GetCurrentPlaybackTimeRange(TRange<FMediaTimeStamp> & T
 		//
 		// No Audio (no data and/or no sink)
 		//
-		if (BlockOnRange.IsEmpty())
+		if (BlockOnRange.GetRange().IsEmpty())
 		{
 			//
 			// Internal clock (DT based)
@@ -2309,7 +2324,7 @@ bool FMediaPlayerFacade::GetCurrentPlaybackTimeRange(TRange<FMediaTimeStamp> & T
 			// External clock delivers time-range
 			// (for now we just use the blocking time range as this clock type is solely used in that case)
 			//
-			TimeRange = BlockOnRange;
+			TimeRange = BlockOnRange.GetRange();
 		}
 	}
 
@@ -2578,7 +2593,8 @@ bool FMediaPlayerFacade::ProcessVideoSamples(IMediaSamples& Samples, const TRang
 		TRange<FMediaTimeStamp> SampleTimeRange(Sample->GetTime(), Sample->GetTime() + Sample->GetDuration());
 
 		// Is it what we want?
-		if (BlockOnRange.IsEmpty() || BlockOnRange.Overlaps(SampleTimeRange))
+		const TRange<FMediaTimeStamp> & BR = BlockOnRange.GetRange();
+		if (BR.IsEmpty() || BR.Overlaps(SampleTimeRange))
 		{
 			// Enqueue the sample to render
 			// (we use a queue to stay compatible with existing structure and older sinks - new sinks will read this single entry right away on the gamethread
@@ -2657,7 +2673,7 @@ void FMediaPlayerFacade::ReceiveMediaEvent(EMediaEvent Event)
 				check(CurrentPlayer->GetPlayerFeatureFlag(IMediaPlayer::EFeatureFlag::UsePlaybackTimingV2));
 
 				// Only do this if we do not block on time ranges
-				if (!BlockOnRange.IsEmpty())
+				if (!BlockOnRange.GetRange().IsEmpty())
 				{
 					// We do not purge as we do not need max perf, but max reliability to actually get certain frames
 					return;
diff --git a/Engine/Source/Runtime/MediaUtils/Public/MediaPlayerFacade.h b/Engine/Source/Runtime/MediaUtils/Public/MediaPlayerFacade.h
--- a/Engine/Source/Runtime/MediaUtils/Public/MediaPlayerFacade.h
+++ b/Engine/Source/Runtime/MediaUtils/Public/MediaPlayerFacade.h
@@ -718,11 +718,6 @@ protected:
 	 */
 	bool BlockOnFetch() const;
 
-	/**
-	 * Internal setup for blocking time ranges
-	 */
-	void SetBlockOnTimeRangeInternal(const TRange<FTimespan>& TimeRange);
-
 	/** Flush all media sample sinks & player plugin. */
 	void Flush(bool bExcludePlayer = false);
 
@@ -834,21 +829,52 @@ private:
 
 private:
 
-	/** The time range to block on sample fetching. */
-	TRange<FMediaTimeStamp> BlockOnRange;
+	class FBlockOnRange
+	{
+	public:
+		FBlockOnRange(FMediaPlayerFacade* InFacade) : Facade(InFacade) { Reset(); }
+
+		void SetRange(const TRange<FTimespan> & NewRange);
+
+		const TRange<FMediaTimeStamp> & GetRange() const;
+
+		void Flush();
+
+		void Reset()
+		{
+			BlockOnRange = TRange<FMediaTimeStamp>::Empty();
+			CurrentTimeRange = TRange<FTimespan>::Empty();
+			LastBlockOnRange = TRange<FTimespan>::Empty();
+			RangeIsDirty = false;
+			OnBlockSeqIndex = 0;
+		}
 
-	/** Last user provided BlockOnRange value */
-	TRange<FTimespan> LastBlockOnRange;
+	private:
+		/** The hosting player facade */
+		FMediaPlayerFacade* Facade;
+
+		/** The last user set time range to block on */
+		TRange<FTimespan> CurrentTimeRange;
+
+		/** The time range to block on sample fetching. */
+		mutable TRange<FMediaTimeStamp> BlockOnRange;
+
+		/** Last user provided BlockOnRange value */
+		mutable TRange<FTimespan> LastBlockOnRange;
+
+		/** Flag to indicate if internal range is valid or not */
+		mutable bool RangeIsDirty;
+
+		/** Sequence index used during blocked playback processing */
+		mutable int64 OnBlockSeqIndex;
+	};
+
+	FBlockOnRange BlockOnRange;
 
-	/** Facade's copy of the last value passed into SetBlockOnRange(); used in case evaluation must be delayed */
-	TRange<FTimespan> BlockOnTimeRange;
 
 	/** Flag to indicate block on range feature as disabled for the current playback session due to previous timeout */
 	bool BlockOnRangeDisabled;
 
-	/** Sequence index used during blocked playback processing */
-	int64 OnBlockSeqIndex;
-
 	/** Media sample cache. */
 	FMediaSampleCache* Cache;
 
