MediaFramework: Thomas's fix for duration assert when playing SimpleComposure.
#rnx

diff --git a/Engine/Source/Runtime/MediaUtils/Private/MediaPlayerFacade.cpp b/Engine/Source/Runtime/MediaUtils/Private/MediaPlayerFacade.cpp
--- a/Engine/Source/Runtime/MediaUtils/Private/MediaPlayerFacade.cpp
+++ b/Engine/Source/Runtime/MediaUtils/Private/MediaPlayerFacade.cpp
@@ -1180,6 +1180,19 @@ void FMediaPlayerFacade::FBlockOnRange::SetRange(const TRange<FTimespan>& NewRan
 }
 
 
+bool FMediaPlayerFacade::FBlockOnRange::IsSet() const
+{
+	TSharedPtr<IMediaPlayer, ESPMode::ThreadSafe> CurrentPlayer(Facade->Player);
+	check(CurrentPlayer.IsValid());
+
+	if (!RangeIsDirty)
+	{
+		return !BlockOnRange.IsEmpty();
+	}
+	return (!CurrentTimeRange.IsEmpty() && CurrentPlayer->GetControls().CanControl(EMediaControl::BlockOnFetch));
+}
+
+
 const TRange<FMediaTimeStamp> & FMediaPlayerFacade::FBlockOnRange::GetRange() const
 {
 	TSharedPtr<IMediaPlayer, ESPMode::ThreadSafe> CurrentPlayer(Facade->Player);
@@ -1200,6 +1213,15 @@ const TRange<FMediaTimeStamp> & FMediaPlayerFacade::FBlockOnRange::GetRange() co
 		return BlockOnRange;
 	}
 
+	EMediaState PlayerState = CurrentPlayer->GetControls().GetState();
+	if (PlayerState != EMediaState::Paused && PlayerState != EMediaState::Playing)
+	{
+		// Return an empty range. Note that the "isSet()" method will still report a set block - so all code will remain in "external clock" mode,
+		// but no samples will be requested (and no actual blocking should take place)
+		static auto EmptyRange(TRange<FMediaTimeStamp>::Empty());
+		return EmptyRange;
+	}
+
 	FTimespan Duration(CurrentPlayer->GetControls().GetDuration());
 	FTimespan Start(CurrentTimeRange.GetLowerBoundValue());
 	FTimespan End(CurrentTimeRange.GetUpperBoundValue());
@@ -2106,7 +2128,7 @@ void FMediaPlayerFacade::PreSampleProcessingTimeHandling()
 	if (!bHaveActiveAudio)
 	{
 		// No external clock? (blocking)
-		if (BlockOnRange.GetRange().IsEmpty())
+		if (!BlockOnRange.IsSet())
 		{
 			// Do we have a current timestamp estimation?
 			if (!NextEstVideoTimeAtFrameStart.IsValid())
@@ -2133,7 +2155,7 @@ void FMediaPlayerFacade::PostSampleProcessingTimeHandling(FTimespan DeltaTime)
 	if (!bHaveActiveAudio)
 	{
 		// No external clock? (blocking)
-		if (BlockOnRange.GetRange().IsEmpty())
+		if (!BlockOnRange.IsSet())
 		{
 			// Move video frame start estimate forward
 			// (the initial NextEstVideoTimeAtFrameStart will never be valid if no video is present)
@@ -2267,7 +2289,7 @@ bool FMediaPlayerFacade::GetCurrentPlaybackTimeRange(TRange<FMediaTimeStamp> & T
 		//
 		// No Audio (no data and/or no sink)
 		//
-		if (BlockOnRange.GetRange().IsEmpty())
+		if (!BlockOnRange.IsSet())
 		{
 			//
 			// Internal clock (DT based)
@@ -2673,7 +2695,7 @@ void FMediaPlayerFacade::ReceiveMediaEvent(EMediaEvent Event)
 				check(CurrentPlayer->GetPlayerFeatureFlag(IMediaPlayer::EFeatureFlag::UsePlaybackTimingV2));
 
 				// Only do this if we do not block on time ranges
-				if (!BlockOnRange.GetRange().IsEmpty())
+				if (BlockOnRange.IsSet())
 				{
 					// We do not purge as we do not need max perf, but max reliability to actually get certain frames
 					return;
diff --git a/Engine/Source/Runtime/MediaUtils/Public/MediaPlayerFacade.h b/Engine/Source/Runtime/MediaUtils/Public/MediaPlayerFacade.h
--- a/Engine/Source/Runtime/MediaUtils/Public/MediaPlayerFacade.h
+++ b/Engine/Source/Runtime/MediaUtils/Public/MediaPlayerFacade.h
@@ -837,6 +837,7 @@ private:
 		void SetRange(const TRange<FTimespan> & NewRange);
 
 		const TRange<FMediaTimeStamp> & GetRange() const;
+		bool IsSet() const;
 
 		void Flush();
 
