Fixing up merge error that brought back an older bug when using blocked/synced playback in MFW
#rnx

diff --git a/Engine/Source/Runtime/MediaUtils/Private/MediaPlayerFacade.cpp b/Engine/Source/Runtime/MediaUtils/Private/MediaPlayerFacade.cpp
--- a/Engine/Source/Runtime/MediaUtils/Private/MediaPlayerFacade.cpp
+++ b/Engine/Source/Runtime/MediaUtils/Private/MediaPlayerFacade.cpp
@@ -677,8 +677,8 @@ bool FMediaPlayerFacade::IsReady() const
 
 	EMediaState State = CurrentPlayer->GetControls().GetState();
 	return ((State != EMediaState::Closed) &&
-			(State != EMediaState::Error) &&
-			(State != EMediaState::Preparing));
+		(State != EMediaState::Error) &&
+		(State != EMediaState::Preparing));
 }
 
 // ----------------------------------------------------------------------------------------------------------------------------------------------
@@ -793,7 +793,7 @@ private:
 
 // ----------------------------------------------------------------------------------------------------------------------------------------------
 
-bool FMediaPlayerFacade::NotifyLifetimeManagerDelegate_PlayerOpen(IMediaPlayerLifecycleManagerDelegate::IControlRef & NewLifecycleManagerDelegateControl, const FString& Url, const IMediaOptions* Options, const FMediaPlayerOptions* PlayerOptions, IMediaPlayerFactory* PlayerFactory, bool bWillCreatePlayer, uint32 WillUseNewResources, uint64 NewPlayerInstanceID)
+bool FMediaPlayerFacade::NotifyLifetimeManagerDelegate_PlayerOpen(IMediaPlayerLifecycleManagerDelegate::IControlRef& NewLifecycleManagerDelegateControl, const FString& Url, const IMediaOptions* Options, const FMediaPlayerOptions* PlayerOptions, IMediaPlayerFactory* PlayerFactory, bool bWillCreatePlayer, uint32 WillUseNewResources, uint64 NewPlayerInstanceID)
 {
 	check(IsInGameThread());
 
@@ -991,7 +991,7 @@ bool FMediaPlayerFacade::ContinueOpen(IMediaPlayerLifecycleManagerDelegate::ICon
 
 			virtual void Signal(uint32 ResourceFlags) override
 			{
-				TFunction<void()> NotifyTask = [TargetDelegateControl=DelegateControl, ResourceFlags]()
+				TFunction<void()> NotifyTask = [TargetDelegateControl = DelegateControl, ResourceFlags]()
 				{
 					// Get MediaModule & check if it is already unloaded...
 					IMediaModule* TargetMediaModule = FModuleManager::GetModulePtr<IMediaModule>("Media");
@@ -1056,16 +1056,16 @@ bool FMediaPlayerFacade::ContinueOpen(IMediaPlayerLifecycleManagerDelegate::ICon
 	}
 
 	{
-	FScopeLock Lock(&LastTimeValuesCS);
+		FScopeLock Lock(&LastTimeValuesCS);
 
-	BlockOnRangeDisabled = false;
-	BlockOnRange.Flush();
-	LastVideoSampleProcessedTimeRange = TRange<FMediaTimeStamp>::Empty();
-	LastAudioSampleProcessedTime.Invalidate();
-	CurrentFrameVideoTimeStamp.Invalidate();
-	CurrentFrameAudioTimeStamp.Invalidate();
+		BlockOnRangeDisabled = false;
+		BlockOnRange.Flush();
+		LastVideoSampleProcessedTimeRange = TRange<FMediaTimeStamp>::Empty();
+		LastAudioSampleProcessedTime.Invalidate();
+		CurrentFrameVideoTimeStamp.Invalidate();
+		CurrentFrameAudioTimeStamp.Invalidate();
 
-	NextEstVideoTimeAtFrameStart.Invalidate();
+		NextEstVideoTimeAtFrameStart.Invalidate();
 	}
 
 	if (bCreateNewPlayer)
@@ -1193,7 +1193,7 @@ bool FMediaPlayerFacade::FBlockOnRange::IsSet() const
 }
 
 
-const TRange<FMediaTimeStamp> & FMediaPlayerFacade::FBlockOnRange::GetRange() const
+const TRange<FMediaTimeStamp>& FMediaPlayerFacade::FBlockOnRange::GetRange() const
 {
 	TSharedPtr<IMediaPlayer, ESPMode::ThreadSafe> CurrentPlayer(Facade->Player);
 	check(CurrentPlayer.IsValid());
@@ -1213,6 +1213,15 @@ const TRange<FMediaTimeStamp> & FMediaPlayerFacade::FBlockOnRange::GetRange() co
 		return BlockOnRange;
 	}
 
+	EMediaState PlayerState = CurrentPlayer->GetControls().GetState();
+	if (PlayerState != EMediaState::Paused && PlayerState != EMediaState::Playing)
+	{
+		// Return an empty range. Note that the "isSet()" method will still report a set block - so all code will remain in "external clock" mode,
+		// but no samples will be requested (and no actual blocking should take place)
+		static auto EmptyRange(TRange<FMediaTimeStamp>::Empty());
+		return EmptyRange;
+	}
+
 	FTimespan Duration(CurrentPlayer->GetControls().GetDuration());
 	FTimespan Start(CurrentTimeRange.GetLowerBoundValue());
 	FTimespan End(CurrentTimeRange.GetUpperBoundValue());
@@ -1428,7 +1437,7 @@ bool FMediaPlayerFacade::BlockOnFetch() const
 {
 	check(Player.IsValid());
 
-	const TRange<FMediaTimeStamp> & BR = BlockOnRange.GetRange();
+	const TRange<FMediaTimeStamp>& BR = BlockOnRange.GetRange();
 
 	if (BR.IsEmpty() || !Player->GetControls().CanControl(EMediaControl::BlockOnFetch) || BlockOnRangeDisabled)
 	{
@@ -1454,7 +1463,7 @@ bool FMediaPlayerFacade::BlockOnFetch() const
 
 			// Compute the "theoretical" next sample range...
 			TRange<FMediaTimeStamp> NextSampleTimeRange = (Rate >= 0.0f) ? TRange<FMediaTimeStamp>(LastVideoSampleProcessedTimeRange.GetUpperBoundValue(), LastVideoSampleProcessedTimeRange.GetUpperBoundValue() + LastVideoSampleProcessedTimeRange.Size<FMediaTimeStamp>().Time)
-																		 : TRange<FMediaTimeStamp>(LastVideoSampleProcessedTimeRange.GetLowerBoundValue() - LastVideoSampleProcessedTimeRange.Size<FMediaTimeStamp>().Time, LastVideoSampleProcessedTimeRange.GetLowerBoundValue());
+				: TRange<FMediaTimeStamp>(LastVideoSampleProcessedTimeRange.GetLowerBoundValue() - LastVideoSampleProcessedTimeRange.Size<FMediaTimeStamp>().Time, LastVideoSampleProcessedTimeRange.GetLowerBoundValue());
 
 			FTimespan Duration = Player->GetControls().GetDuration();
 
@@ -1760,7 +1769,7 @@ void FMediaPlayerFacade::ProcessEvent(EMediaEvent Event)
 			// Yes, this also means: if we still have a player, it's still the one this event originated from
 
 			// If player allows: close it down all the way right now
-			if(Player.IsValid() && Player->GetPlayerFeatureFlag(IMediaPlayer::EFeatureFlag::AllowShutdownOnClose))
+			if (Player.IsValid() && Player->GetPlayerFeatureFlag(IMediaPlayer::EFeatureFlag::AllowShutdownOnClose))
 			{
 				bDidRecentPlayerHaveError = HasError();
 				DestroyPlayer();
@@ -2070,8 +2079,8 @@ void FMediaPlayerFacade::TickTickable()
 	}
 
 	{
-	FScopeLock Lock1(&LastTimeValuesCS);
-	Player->SetLastAudioRenderedSampleTime(LastAudioRenderedSampleTime.TimeStamp.Time);
+		FScopeLock Lock1(&LastTimeValuesCS);
+		Player->SetLastAudioRenderedSampleTime(LastAudioRenderedSampleTime.TimeStamp.Time);
 	}
 
 	Player->TickAudio();
@@ -2209,7 +2218,7 @@ void FMediaPlayerFacade::PostSampleProcessingTimeHandling(FTimespan DeltaTime)
 }
 
 
-bool FMediaPlayerFacade::GetCurrentPlaybackTimeRange(TRange<FMediaTimeStamp> & TimeRange, float Rate, FTimespan DeltaTime, bool bDoNotUseFrameStartReference) const
+bool FMediaPlayerFacade::GetCurrentPlaybackTimeRange(TRange<FMediaTimeStamp>& TimeRange, float Rate, FTimespan DeltaTime, bool bDoNotUseFrameStartReference) const
 {
 	check(Player->GetPlayerFeatureFlag(IMediaPlayer::EFeatureFlag::UsePlaybackTimingV2));
 
@@ -2273,7 +2282,7 @@ bool FMediaPlayerFacade::GetCurrentPlaybackTimeRange(TRange<FMediaTimeStamp> & T
 		}
 
 		TimeRange = (Rate >= 0.0f) ? TRange<FMediaTimeStamp>(EstAudioTimeAtFrameStart, EstAudioTimeAtFrameStart + DeltaTime * Rate)
-								   : TRange<FMediaTimeStamp>(EstAudioTimeAtFrameStart + DeltaTime * (1.0f + Rate), EstAudioTimeAtFrameStart + DeltaTime);
+			: TRange<FMediaTimeStamp>(EstAudioTimeAtFrameStart + DeltaTime * (1.0f + Rate), EstAudioTimeAtFrameStart + DeltaTime);
 	}
 	else
 	{
@@ -2308,7 +2317,7 @@ bool FMediaPlayerFacade::GetCurrentPlaybackTimeRange(TRange<FMediaTimeStamp> & T
 				}
 
 				TimeRange = (Rate >= 0.0f) ? TRange<FMediaTimeStamp>(NextEstVideoTimeAtFrameStart.TimeStamp, NextEstVideoTimeAtFrameStart.TimeStamp + DeltaTime * Rate)
-										   : TRange<FMediaTimeStamp>(NextEstVideoTimeAtFrameStart.TimeStamp + DeltaTime * (1.0f + Rate), NextEstVideoTimeAtFrameStart.TimeStamp + DeltaTime);
+					: TRange<FMediaTimeStamp>(NextEstVideoTimeAtFrameStart.TimeStamp + DeltaTime * (1.0f + Rate), NextEstVideoTimeAtFrameStart.TimeStamp + DeltaTime);
 
 				// If we are looping we check to prepare proper ranges should we wrap around either end of the media...
 				// (we do not clamp in the non-looping case as the rest of the code should deal with that fine)
@@ -2374,9 +2383,9 @@ void FMediaPlayerFacade::ProcessAudioSamples(IMediaSamples& Samples, TRange<FTim
 				}
 
 				{
-				FScopeLock Lock(&LastTimeValuesCS);
-				LastAudioSampleProcessedTime.TimeStamp = FMediaTimeStamp(Sample->GetTime());
-				LastAudioSampleProcessedTime.SampledAtTime = FPlatformTime::Seconds();
+					FScopeLock Lock(&LastTimeValuesCS);
+					LastAudioSampleProcessedTime.TimeStamp = FMediaTimeStamp(Sample->GetTime());
+					LastAudioSampleProcessedTime.SampledAtTime = FPlatformTime::Seconds();
 				}
 
 				PinnedPrimaryAudioSink->Enqueue(Sample.ToSharedRef());
@@ -2390,8 +2399,8 @@ void FMediaPlayerFacade::ProcessAudioSamples(IMediaSamples& Samples, TRange<FTim
 				// We got video and audio, but no audio sink - throw away anything up to video playback time...
 				// (rough estimate, as this is off-gamethread; but better than throwing things out with no throttling at all)
 				{
-				FScopeLock Lock(&LastTimeValuesCS);
-				TimeRange.SetUpperBound(TRangeBound<FTimespan>(CurrentFrameVideoTimeStamp.Time));
+					FScopeLock Lock(&LastTimeValuesCS);
+					TimeRange.SetUpperBound(TRangeBound<FTimespan>(CurrentFrameVideoTimeStamp.Time));
 				}
 				while (Samples.FetchAudio(TimeRange, Sample))
 					;
@@ -2442,9 +2451,9 @@ void FMediaPlayerFacade::ProcessCaptionSamplesV1(IMediaSamples& Samples, TRange<
 	{
 		if (Sample.IsValid() && !CaptionSampleSinks.Enqueue(Sample.ToSharedRef(), FMediaPlayerQueueDepths::MaxCaptionSinkDepth))
 		{
-			#if MEDIAPLAYERFACADE_TRACE_SINKOVERFLOWS
-				UE_LOG(LogMediaUtils, VeryVerbose, TEXT("PlayerFacade %p: Caption sample sink overflow"), this);
-			#endif
+#if MEDIAPLAYERFACADE_TRACE_SINKOVERFLOWS
+			UE_LOG(LogMediaUtils, VeryVerbose, TEXT("PlayerFacade %p: Caption sample sink overflow"), this);
+#endif
 		}
 	}
 }
@@ -2458,9 +2467,9 @@ void FMediaPlayerFacade::ProcessMetadataSamples(IMediaSamples& Samples, TRange<F
 	{
 		if (Sample.IsValid() && !MetadataSampleSinks.Enqueue(Sample.ToSharedRef(), FMediaPlayerQueueDepths::MaxMetadataSinkDepth))
 		{
-			#if MEDIAPLAYERFACADE_TRACE_SINKOVERFLOWS
-				UE_LOG(LogMediaUtils, VeryVerbose, TEXT("PlayerFacade %p: Metadata sample sink overflow"), this);
-			#endif
+#if MEDIAPLAYERFACADE_TRACE_SINKOVERFLOWS
+			UE_LOG(LogMediaUtils, VeryVerbose, TEXT("PlayerFacade %p: Metadata sample sink overflow"), this);
+#endif
 		}
 	}
 }
@@ -2474,9 +2483,9 @@ void FMediaPlayerFacade::ProcessSubtitleSamplesV1(IMediaSamples& Samples, TRange
 	{
 		if (Sample.IsValid() && !SubtitleSampleSinks.Enqueue(Sample.ToSharedRef(), FMediaPlayerQueueDepths::MaxSubtitleSinkDepth))
 		{
-			#if MEDIAPLAYERFACADE_TRACE_SINKOVERFLOWS
-				UE_LOG(LogMediaUtils, VeryVerbose, TEXT("PlayerFacade %p: Subtitle sample sink overflow"), this);
-			#endif
+#if MEDIAPLAYERFACADE_TRACE_SINKOVERFLOWS
+			UE_LOG(LogMediaUtils, VeryVerbose, TEXT("PlayerFacade %p: Subtitle sample sink overflow"), this);
+#endif
 		}
 	}
 }
@@ -2503,8 +2512,8 @@ void FMediaPlayerFacade::ProcessVideoSamplesV1(IMediaSamples& Samples, TRange<FT
 		}
 
 		{
-		FScopeLock Lock(&LastTimeValuesCS);
-		CurrentFrameVideoTimeStamp = Sample->GetTime();
+			FScopeLock Lock(&LastTimeValuesCS);
+			CurrentFrameVideoTimeStamp = Sample->GetTime();
 		}
 
 		UE_LOG(LogMediaUtils, VeryVerbose, TEXT("PlayerFacade %p: Fetched video sample %s"), this, *Sample->GetTime().Time.ToString(TEXT("%h:%m:%s.%t")));
@@ -2527,7 +2536,7 @@ void FMediaPlayerFacade::ProcessVideoSamplesV1(IMediaSamples& Samples, TRange<FT
 }
 
 
-bool FMediaPlayerFacade::ProcessVideoSamples(IMediaSamples& Samples, const TRange<FMediaTimeStamp> & TimeRange)
+bool FMediaPlayerFacade::ProcessVideoSamples(IMediaSamples& Samples, const TRange<FMediaTimeStamp>& TimeRange)
 {
 	// Let the player do some processing if needed.
 	if (Player.IsValid())
@@ -2552,39 +2561,39 @@ bool FMediaPlayerFacade::ProcessVideoSamples(IMediaSamples& Samples, const TRang
 
 		switch (Samples.FetchBestVideoSampleForTimeRange(TimeRange, Sample, bReverse))
 		{
-			case IMediaSamples::EFetchBestSampleResult::Ok:
-				break;
+		case IMediaSamples::EFetchBestSampleResult::Ok:
+			break;
 
-			case IMediaSamples::EFetchBestSampleResult::NoSample:
+		case IMediaSamples::EFetchBestSampleResult::NoSample:
+		{
+			break;
+		}
+
+		case IMediaSamples::EFetchBestSampleResult::NotSupported:
+		{
+			//
+			// Fallback for players supporting V2 timing, but do not supply FetchBestVideoSampleForTimeRange() due to some
+			// custom implementation of IMediaSamples (here to ease adoption of the new timing code - eventually should go away)
+			//
+
+			// Find newest sample that satisfies the time range
+			// (the FetchXYZ() code does not work well with a lower range limit at all - we ask for a "up to" type range instead
+			//  and limit the other side of the range in code here to not change the older logic & possibly cause trouble in old code)
+			TRange<FMediaTimeStamp> TempRange = bReverse ? TRange<FMediaTimeStamp>::AtLeast(TimeRange.GetUpperBoundValue()) : TRange<FMediaTimeStamp>::AtMost(TimeRange.GetUpperBoundValue());
+			while (Samples.FetchVideo(TempRange, Sample))
+				;
+			if (Sample.IsValid() &&
+				((!bReverse && ((Sample->GetTime() + Sample->GetDuration()) > TimeRange.GetLowerBoundValue())) ||
+					(bReverse && ((Sample->GetTime() - Sample->GetDuration()) < TimeRange.GetLowerBoundValue()))))
 			{
-				break;
+				// Sample is good - nothing more to do here
 			}
-
-			case IMediaSamples::EFetchBestSampleResult::NotSupported:
+			else
 			{
-				//
-				// Fallback for players supporting V2 timing, but do not supply FetchBestVideoSampleForTimeRange() due to some
-				// custom implementation of IMediaSamples (here to ease adoption of the new timing code - eventually should go away)
-				//
-
-				// Find newest sample that satisfies the time range
-				// (the FetchXYZ() code does not work well with a lower range limit at all - we ask for a "up to" type range instead
-				//  and limit the other side of the range in code here to not change the older logic & possibly cause trouble in old code)
-				TRange<FMediaTimeStamp> TempRange = bReverse ? TRange<FMediaTimeStamp>::AtLeast(TimeRange.GetUpperBoundValue()) : TRange<FMediaTimeStamp>::AtMost(TimeRange.GetUpperBoundValue());
-				while (Samples.FetchVideo(TempRange, Sample))
-					;
-				if (Sample.IsValid() &&
-					((!bReverse && ((Sample->GetTime() + Sample->GetDuration()) > TimeRange.GetLowerBoundValue())) ||
-					(bReverse && ((Sample->GetTime() - Sample->GetDuration()) < TimeRange.GetLowerBoundValue()))))
-				{
-					// Sample is good - nothing more to do here
-				}
-				else
-				{
-					Sample.Reset();
-				}
-				break;
+				Sample.Reset();
 			}
+			break;
+		}
 		}
 	}
 	else
@@ -2606,7 +2615,7 @@ bool FMediaPlayerFacade::ProcessVideoSamples(IMediaSamples& Samples, const TRang
 		TRange<FMediaTimeStamp> SampleTimeRange(Sample->GetTime(), Sample->GetTime() + Sample->GetDuration());
 
 		// Is it what we want?
-		const TRange<FMediaTimeStamp> & BR = BlockOnRange.GetRange();
+		const TRange<FMediaTimeStamp>& BR = BlockOnRange.GetRange();
 		if (BR.IsEmpty() || BR.Overlaps(SampleTimeRange))
 		{
 			// Enqueue the sample to render
@@ -2667,98 +2676,98 @@ void FMediaPlayerFacade::ReceiveMediaEvent(EMediaEvent Event)
 
 	if (Event >= EMediaEvent::Internal_Start)
 	{
-		switch(Event)
+		switch (Event)
 		{
-			case	EMediaEvent::Internal_PurgeVideoSamplesHint:
-			{
-				//
-				// Player asks to attempt to purge older samples in the video output queue it maintains
-				// (ask goes via facade as the player does not have accurate timing info)
-				//
-				TSharedPtr<IMediaPlayer, ESPMode::ThreadSafe> CurrentPlayer = Player;
-
-				if (!CurrentPlayer.IsValid())
-				{
-					return;
-				}
-
-				// We only support this for V2 timing players
-				check(CurrentPlayer->GetPlayerFeatureFlag(IMediaPlayer::EFeatureFlag::UsePlaybackTimingV2));
-
-				// Only do this if we do not block on time ranges
-				if (BlockOnRange.IsSet())
-				{
-					// We do not purge as we do not need max perf, but max reliability to actually get certain frames
-					return;
-				}
-
-				float Rate = CurrentRate;
-				if (Rate == 0.0f)
-				{
-					return;
-				}
-
-				// Get current playback time
-				// (note: we have DeltaTime forced to zero -> we just get a single value & we compute relative to "now", noty any game frame start)
-				TRange<FMediaTimeStamp> TimeRange;
-				if (!GetCurrentPlaybackTimeRange(TimeRange, Rate, FTimespan::Zero(), true))
-				{
-					return;
-				}
-
-				bool bReverse = (Rate < 0.0f);
-				uint32 NumPurged = CurrentPlayer->GetSamples().PurgeOutdatedVideoSamples(TimeRange.GetLowerBoundValue() + (bReverse ? kOutdatedVideoSamplesTollerance : -kOutdatedVideoSamplesTollerance), bReverse);
-
-				SET_DWORD_STAT(STAT_MediaUtils_FacadeNumPurgedVideoSamples, NumPurged);
-				INC_DWORD_STAT_BY(STAT_MediaUtils_FacadeTotalPurgedVideoSamples, NumPurged);
-
-				break;
-			}
+		case	EMediaEvent::Internal_PurgeVideoSamplesHint:
+		{
+			//
+			// Player asks to attempt to purge older samples in the video output queue it maintains
+			// (ask goes via facade as the player does not have accurate timing info)
+			//
+			TSharedPtr<IMediaPlayer, ESPMode::ThreadSafe> CurrentPlayer = Player;
 
-			case	EMediaEvent::Internal_ResetForDiscontinuity:
+			if (!CurrentPlayer.IsValid())
 			{
-				UE_LOG(LogMediaUtils, VeryVerbose, TEXT("PlayerFacade %p: Reset for discontinuity"), this);
-				bIsSinkFlushPending = true;
-				break;
-			}
-			case	EMediaEvent::Internal_RenderClockStart:
-			{
-				UE_LOG(LogMediaUtils, VeryVerbose, TEXT("PlayerFacade %p: Render clock shall start"), this);
-				break;
-			}
-			case	EMediaEvent::Internal_RenderClockStop:
-			{
-				UE_LOG(LogMediaUtils, VeryVerbose, TEXT("PlayerFacade %p: Render clock shall stop"), this);
-				break;
+				return;
 			}
 
-			case	EMediaEvent::Internal_VideoSamplesAvailable:
-			{
-				UE_LOG(LogMediaUtils, VeryVerbose, TEXT("PlayerFacade %p: Video samples ARE available"), this);
-				VideoSampleAvailability = 1;
-				break;
-			}
-			case	EMediaEvent::Internal_VideoSamplesUnavailable:
+			// We only support this for V2 timing players
+			check(CurrentPlayer->GetPlayerFeatureFlag(IMediaPlayer::EFeatureFlag::UsePlaybackTimingV2));
+
+			// Only do this if we do not block on time ranges
+			if (BlockOnRange.IsSet())
 			{
-				UE_LOG(LogMediaUtils, VeryVerbose, TEXT("PlayerFacade %p: Video samples are NOT available"), this);
-				VideoSampleAvailability = 0;
-				break;
+				// We do not purge as we do not need max perf, but max reliability to actually get certain frames
+				return;
 			}
-			case	EMediaEvent::Internal_AudioSamplesAvailable:
+
+			float Rate = CurrentRate;
+			if (Rate == 0.0f)
 			{
-				UE_LOG(LogMediaUtils, VeryVerbose, TEXT("PlayerFacade %p: Audio samples ARE available"), this);
-				AudioSampleAvailability = 1;
-				break;
+				return;
 			}
-			case	EMediaEvent::Internal_AudioSamplesUnavailable:
+
+			// Get current playback time
+			// (note: we have DeltaTime forced to zero -> we just get a single value & we compute relative to "now", noty any game frame start)
+			TRange<FMediaTimeStamp> TimeRange;
+			if (!GetCurrentPlaybackTimeRange(TimeRange, Rate, FTimespan::Zero(), true))
 			{
-				UE_LOG(LogMediaUtils, VeryVerbose, TEXT("PlayerFacade %p: Audio samples are NOT available"), this);
-				AudioSampleAvailability = 0;
-				break;
+				return;
 			}
 
-			default:
-				break;
+			bool bReverse = (Rate < 0.0f);
+			uint32 NumPurged = CurrentPlayer->GetSamples().PurgeOutdatedVideoSamples(TimeRange.GetLowerBoundValue() + (bReverse ? kOutdatedVideoSamplesTollerance : -kOutdatedVideoSamplesTollerance), bReverse);
+
+			SET_DWORD_STAT(STAT_MediaUtils_FacadeNumPurgedVideoSamples, NumPurged);
+			INC_DWORD_STAT_BY(STAT_MediaUtils_FacadeTotalPurgedVideoSamples, NumPurged);
+
+			break;
+		}
+
+		case	EMediaEvent::Internal_ResetForDiscontinuity:
+		{
+			UE_LOG(LogMediaUtils, VeryVerbose, TEXT("PlayerFacade %p: Reset for discontinuity"), this);
+			bIsSinkFlushPending = true;
+			break;
+		}
+		case	EMediaEvent::Internal_RenderClockStart:
+		{
+			UE_LOG(LogMediaUtils, VeryVerbose, TEXT("PlayerFacade %p: Render clock shall start"), this);
+			break;
+		}
+		case	EMediaEvent::Internal_RenderClockStop:
+		{
+			UE_LOG(LogMediaUtils, VeryVerbose, TEXT("PlayerFacade %p: Render clock shall stop"), this);
+			break;
+		}
+
+		case	EMediaEvent::Internal_VideoSamplesAvailable:
+		{
+			UE_LOG(LogMediaUtils, VeryVerbose, TEXT("PlayerFacade %p: Video samples ARE available"), this);
+			VideoSampleAvailability = 1;
+			break;
+		}
+		case	EMediaEvent::Internal_VideoSamplesUnavailable:
+		{
+			UE_LOG(LogMediaUtils, VeryVerbose, TEXT("PlayerFacade %p: Video samples are NOT available"), this);
+			VideoSampleAvailability = 0;
+			break;
+		}
+		case	EMediaEvent::Internal_AudioSamplesAvailable:
+		{
+			UE_LOG(LogMediaUtils, VeryVerbose, TEXT("PlayerFacade %p: Audio samples ARE available"), this);
+			AudioSampleAvailability = 1;
+			break;
+		}
+		case	EMediaEvent::Internal_AudioSamplesUnavailable:
+		{
+			UE_LOG(LogMediaUtils, VeryVerbose, TEXT("PlayerFacade %p: Audio samples are NOT available"), this);
+			AudioSampleAvailability = 0;
+			break;
+		}
+
+		default:
+			break;
 		}
 	}
 	else
