MediaFramework: Changed more game thread checks to allow for the slate thread as well.
#rnx

diff --git a/Engine/Source/Runtime/MediaUtils/Private/MediaPlayerFacade.cpp b/Engine/Source/Runtime/MediaUtils/Private/MediaPlayerFacade.cpp
--- a/Engine/Source/Runtime/MediaUtils/Private/MediaPlayerFacade.cpp
+++ b/Engine/Source/Runtime/MediaUtils/Private/MediaPlayerFacade.cpp
@@ -795,7 +795,7 @@ private:
 
 bool FMediaPlayerFacade::NotifyLifetimeManagerDelegate_PlayerOpen(IMediaPlayerLifecycleManagerDelegate::IControlRef & NewLifecycleManagerDelegateControl, const FString& Url, const IMediaOptions* Options, const FMediaPlayerOptions* PlayerOptions, IMediaPlayerFactory* PlayerFactory, bool bWillCreatePlayer, uint32 WillUseNewResources, uint64 NewPlayerInstanceID)
 {
-	check(IsInGameThread());
+	check(IsInGameThread() || IsInSlateThread());
 
 	if (IMediaPlayerLifecycleManagerDelegate* Delegate = MediaModule->GetPlayerLifecycleManagerDelegate())
 	{
@@ -821,7 +821,7 @@ bool FMediaPlayerFacade::NotifyLifetimeManagerDelegate_PlayerOpen(IMediaPlayerLi
 
 bool FMediaPlayerFacade::NotifyLifetimeManagerDelegate_PlayerCreated()
 {
-	check(IsInGameThread());
+	check(IsInGameThread() || IsInSlateThread());
 	check(Player.IsValid());
 
 	if (LifecycleManagerDelegateControl.IsValid())
@@ -837,7 +837,7 @@ bool FMediaPlayerFacade::NotifyLifetimeManagerDelegate_PlayerCreated()
 
 bool FMediaPlayerFacade::NotifyLifetimeManagerDelegate_PlayerCreateFailed()
 {
-	check(IsInGameThread());
+	check(IsInGameThread() || IsInSlateThread());
 
 	if (LifecycleManagerDelegateControl.IsValid())
 	{
@@ -852,7 +852,7 @@ bool FMediaPlayerFacade::NotifyLifetimeManagerDelegate_PlayerCreateFailed()
 
 bool FMediaPlayerFacade::NotifyLifetimeManagerDelegate_PlayerClosed()
 {
-	check(IsInGameThread());
+	check(IsInGameThread() || IsInSlateThread());
 
 	if (LifecycleManagerDelegateControl.IsValid())
 	{
@@ -867,7 +867,7 @@ bool FMediaPlayerFacade::NotifyLifetimeManagerDelegate_PlayerClosed()
 
 bool FMediaPlayerFacade::NotifyLifetimeManagerDelegate_PlayerDestroyed()
 {
-	check(IsInGameThread());
+	check(IsInGameThread() || IsInSlateThread());
 
 	if (LifecycleManagerDelegateControl.IsValid())
 	{
@@ -882,7 +882,7 @@ bool FMediaPlayerFacade::NotifyLifetimeManagerDelegate_PlayerDestroyed()
 
 bool FMediaPlayerFacade::NotifyLifetimeManagerDelegate_PlayerResourcesReleased(uint32 ResourceFlags)
 {
-	check(IsInGameThread());
+	check(IsInGameThread() || IsInSlateThread());
 
 	if (LifecycleManagerDelegateControl.IsValid())
 	{
@@ -2237,7 +2237,7 @@ bool FMediaPlayerFacade::GetCurrentPlaybackTimeRange(TRange<FMediaTimeStamp> & T
 			// Normal estimation relative to current frame start...
 			// (on gamethread operation)
 
-			check(IsInGameThread());
+			check(IsInGameThread() || IsInSlateThread());
 
 			double AgeOfFrameStart = Now - MediaModule->GetFrameStartTime();
 			double AgeOfAudioTime = Now - AudioTime.SampledAtTime;
