Added comparison operators to MediaPlayerOptions
#rnx

diff --git a/Engine/Source/Runtime/MediaUtils/Public/MediaPlayerOptions.h b/Engine/Source/Runtime/MediaUtils/Public/MediaPlayerOptions.h
--- a/Engine/Source/Runtime/MediaUtils/Public/MediaPlayerOptions.h
+++ b/Engine/Source/Runtime/MediaUtils/Public/MediaPlayerOptions.h
@@ -32,6 +32,18 @@ struct FMediaPlayerTrackOptions
 	{
 	}
 
+	/** Check if this Media Player Track Options is not equal to another one. */
+	bool operator!=(const FMediaPlayerTrackOptions& Other) const
+	{
+		return (Audio != Other.Audio
+			|| Caption != Other.Caption
+			|| Metadata != Other.Metadata
+			|| Script != Other.Script
+			|| Subtitle != Other.Subtitle
+			|| Text != Other.Text
+			|| Video != Other.Video);
+	}
+
 	UPROPERTY(BlueprintReadWrite, Category = "Tracks")
 	int32 Audio;
 
@@ -66,6 +78,15 @@ struct FMediaPlayerOptions
 	{
 	}
 
+	/** Check if this Media Options is not equal to another one. */
+	bool operator!=(const FMediaPlayerOptions& Other) const
+	{
+		return (Tracks != Other.Tracks
+			|| SeekTime != Other.SeekTime
+			|| PlayOnOpen != Other.PlayOnOpen
+			|| Loop != Other.Loop);
+	}
+
 	UPROPERTY(BlueprintReadWrite, Category = "Tracks")
 	FMediaPlayerTrackOptions Tracks;
 
