Locking the player instance through a local TSharedPtr to avoid it being deleted with the member class pointer becoming invalid during processing.

#rnx

diff --git a/Engine/Source/Runtime/MediaUtils/Private/MediaPlayerFacade.cpp b/Engine/Source/Runtime/MediaUtils/Private/MediaPlayerFacade.cpp
--- a/Engine/Source/Runtime/MediaUtils/Private/MediaPlayerFacade.cpp
+++ b/Engine/Source/Runtime/MediaUtils/Private/MediaPlayerFacade.cpp
@@ -1977,7 +1977,8 @@ void FMediaPlayerFacade::TickFetch(FTimespan DeltaTime, FTimespan Timecode)
 {
 	SCOPE_CYCLE_COUNTER(STAT_MediaUtils_FacadeTickFetch);
 
-	if (!Player.IsValid())
+	TSharedPtr<IMediaPlayer, ESPMode::ThreadSafe> CurrentPlayer(Player);
+	if (!CurrentPlayer.IsValid())
 	{
 		// Send out deferred broadcasts.
 		EMediaEvent Event;
@@ -1998,14 +1999,14 @@ void FMediaPlayerFacade::TickFetch(FTimespan DeltaTime, FTimespan Timecode)
 		return;
 	}
 
-	if (!Player->GetPlayerFeatureFlag(IMediaPlayer::EFeatureFlag::UsePlaybackTimingV2))
+	if (!CurrentPlayer->GetPlayerFeatureFlag(IMediaPlayer::EFeatureFlag::UsePlaybackTimingV2))
 	{
 		//
 		// Old timing control
 		//
 
 		// let the player generate samples & process events
-		Player->TickFetch(DeltaTime, Timecode);
+		CurrentPlayer->TickFetch(DeltaTime, Timecode);
 
 		{
 			// process deferred events
@@ -2039,7 +2040,7 @@ void FMediaPlayerFacade::TickFetch(FTimespan DeltaTime, FTimespan Timecode)
 		}
 
 		// process samples in range
-		IMediaSamples& Samples = Player->GetSamples();
+		IMediaSamples& Samples = CurrentPlayer->GetSamples();
 
 		bool Blocked = false;
 		FDateTime BlockedTime;
