Fixed uninitialized data access reported by UBSan.

#rnx

diff --git a/Engine/Source/Runtime/MediaUtils/Private/MediaPlayerFacade.cpp b/Engine/Source/Runtime/MediaUtils/Private/MediaPlayerFacade.cpp
--- a/Engine/Source/Runtime/MediaUtils/Private/MediaPlayerFacade.cpp
+++ b/Engine/Source/Runtime/MediaUtils/Private/MediaPlayerFacade.cpp
@@ -108,12 +108,12 @@ FMediaPlayerFacade::FMediaPlayerFacade()
 	, bHaveActiveAudio(false)
 	, VideoSampleAvailability(-1)
 	, AudioSampleAvailability(-1)
+	, bIsSinkFlushPending(false)
 	, bAreEventsSafeForAnyThread(false)
 {
 	BlockOnRangeDisabled = false;
 
 	MediaModule = FModuleManager::LoadModulePtr<IMediaModule>("Media");
-	bIsSinkFlushPending = false;
 	bDidRecentPlayerHaveError = false;
 }
 
