MediaFramework: Added URL to warning message when blocking playback is disabled so the offending player can be identified.
#rnx

diff --git a/Engine/Source/Runtime/MediaUtils/Private/MediaPlayerFacade.cpp b/Engine/Source/Runtime/MediaUtils/Private/MediaPlayerFacade.cpp
--- a/Engine/Source/Runtime/MediaUtils/Private/MediaPlayerFacade.cpp
+++ b/Engine/Source/Runtime/MediaUtils/Private/MediaPlayerFacade.cpp
@@ -1944,7 +1944,12 @@ void FMediaPlayerFacade::TickInput(FTimespan DeltaTime, FTimespan Timecode)
 				// Timeout?
 				if ((FPlatformTime::Seconds() - BlockingStart) > MEDIAUTILS_MAX_BLOCKONFETCH_SECONDS)
 				{
-					UE_LOG(LogMediaUtils, Error, TEXT("Blocking media playback timed out. Disabling it for this playback session."));
+					FString Url;
+#if !UE_BUILD_SHIPPING
+					Url = Player->GetUrl();
+#endif // !UE_BUILD_SHIPPING
+					UE_LOG(LogMediaUtils, Error, TEXT("Blocking media playback timed out. Disabling it for this playback session. URL:%s"),
+						*Url);
 					BlockOnRangeDisabled = true;
 					break;
 				}
