Adding additional precautions to deal with slightly "off-sh" blocked ranges being set to control "blocked synced" plabck in the MFW player facade

diff --git a/Engine/Source/Runtime/MediaUtils/Private/MediaPlayerFacade.cpp b/Engine/Source/Runtime/MediaUtils/Private/MediaPlayerFacade.cpp
--- a/Engine/Source/Runtime/MediaUtils/Private/MediaPlayerFacade.cpp
+++ b/Engine/Source/Runtime/MediaUtils/Private/MediaPlayerFacade.cpp
@@ -32,6 +32,8 @@
 
 #include "Async/Async.h"
 
+#include <algorithm>
+
 #define MEDIAPLAYERFACADE_DISABLE_BLOCKING 0
 #define MEDIAPLAYERFACADE_TRACE_SINKOVERFLOWS 0
 
@@ -1245,6 +1247,8 @@ const TRange<FMediaTimeStamp>& FMediaPlayerFacade::FBlockOnRange::GetRange() con
 	FTimespan Start(CurrentTimeRange.GetLowerBoundValue());
 	FTimespan End(CurrentTimeRange.GetUpperBoundValue());
 
+	auto SetBlockOnRange = BlockOnRange;
+
 	if (!CurrentPlayer->GetControls().IsLooping())
 	{
 		// We pass in the time range as is on seq-index zero at all times - players have to reject sample output / blocking logic will detect begin outside media range
@@ -1315,6 +1319,21 @@ const TRange<FMediaTimeStamp>& FMediaPlayerFacade::FBlockOnRange::GetRange() con
 		check(!BlockOnRange.IsEmpty());
 	}
 
+	// Does the new range overlap with the last one we set?
+	if (!SetBlockOnRange.IsEmpty() && SetBlockOnRange.Overlaps(BlockOnRange))
+	{
+		// Yes, make sure the new range is setup so that it is a "correct" progression given the current playback direction...
+		// (this may go so far as to undo any updates if the "new" one does not adds any range in the playback direction)
+		if (CurrentPlayer->GetControls().GetRate() >= 0.0f)
+		{
+			BlockOnRange = TRange<FMediaTimeStamp>(std::max(BlockOnRange.GetLowerBoundValue(), SetBlockOnRange.GetLowerBoundValue()), std::max(BlockOnRange.GetUpperBoundValue(), SetBlockOnRange.GetUpperBoundValue()));
+		}
+		else
+		{
+			BlockOnRange = TRange<FMediaTimeStamp>(std::min(BlockOnRange.GetLowerBoundValue(), SetBlockOnRange.GetLowerBoundValue()), std::min(BlockOnRange.GetUpperBoundValue(), SetBlockOnRange.GetUpperBoundValue()));
+		}
+	}
+
 	CurrentPlayer->GetControls().SetBlockingPlaybackHint(!BlockOnRange.IsEmpty());
 
 	LastBlockOnRange = CurrentTimeRange;
@@ -1479,6 +1498,7 @@ bool FMediaPlayerFacade::BlockOnFetch() const
 		if (LastVideoSampleProcessedTimeRange.Overlaps(BR))
 		{
 			// We have no new data (else we would not even call this method), but the last sample we returned is still inside the current range -> good, but...
+			//
 			// If the next sample would already cover more of the range than the older one we would like to use that instead -> but it may well be we do not have any data about the sample yet (and would indeed LIKE to block!)
 			// So, we assume that the next sample will follow with no gap and have the same duration (a pretty good, general assumption) and check against that data to see if it would be better...
 
