Fixing blocked/synced p[layback timeouts with enabled audio playback on player setup (syncing will be disabled)

diff --git a/Engine/Source/Runtime/MediaUtils/Private/MediaPlayerFacade.cpp b/Engine/Source/Runtime/MediaUtils/Private/MediaPlayerFacade.cpp
--- a/Engine/Source/Runtime/MediaUtils/Private/MediaPlayerFacade.cpp
+++ b/Engine/Source/Runtime/MediaUtils/Private/MediaPlayerFacade.cpp
@@ -1492,6 +1492,13 @@ bool FMediaPlayerFacade::BlockOnFetch() const
 		//
 		// V2 blocking logic
 		//
+
+		// If we have any active audio playback we skip any blocking
+		if (HaveAudioPlayback())
+		{
+			return false;
+		}
+
 		float Rate = GetUnpausedRate();
 
 		// If the current sample "out there" is actually overlapping with the current block, we might be good with no new sample
