MediaUtils: Fixed incorrect threshold value due to implicit type casting.

#rnx

diff --git a/Engine/Source/Runtime/MediaUtils/Private/MediaPlayerFacade.cpp b/Engine/Source/Runtime/MediaUtils/Private/MediaPlayerFacade.cpp
--- a/Engine/Source/Runtime/MediaUtils/Private/MediaPlayerFacade.cpp
+++ b/Engine/Source/Runtime/MediaUtils/Private/MediaPlayerFacade.cpp
@@ -547,7 +547,7 @@ FString FMediaPlayerFacade::GetTrackLanguage(EMediaTrackType TrackType, int32 Tr
 float FMediaPlayerFacade::GetVideoTrackAspectRatio(int32 TrackIndex, int32 FormatIndex) const
 {
 	FMediaVideoTrackFormat Format;
-	return (GetVideoTrackFormat(TrackIndex, FormatIndex, Format) && (Format.Dim.Y != 0)) ? ((float)(Format.Dim.X) / Format.Dim.Y) : 0.0f;
+	return (GetVideoTrackFormat(TrackIndex, FormatIndex, Format) && (Format.Dim.Y != 0)) ? ((float)(Format.Dim.X) / (float)Format.Dim.Y) : 0.0f;
 }
 
 
@@ -2748,7 +2748,7 @@ void FMediaPlayerFacade::ProcessSubtitleSamples(IMediaSamples& Samples, TRange<F
 	if (CurrentPlayer.IsValid() && CurrentPlayer->GetPlayerFeatureFlag(IMediaPlayer::EFeatureFlag::UsePlaybackTimingV2))
 	{
 		bool bReverse = CurrentRate < 0.0f;
-		uint32 NumPurged = CurrentPlayer->GetSamples().PurgeOutdatedSubtitleSamples(TimeRange.GetLowerBoundValue() + (bReverse ? kOutdatedVideoSamplesTolerance : -kOutdatedVideoSamplesTolerance), bReverse);
+		uint32 NumPurged = CurrentPlayer->GetSamples().PurgeOutdatedSubtitleSamples(TimeRange.GetLowerBoundValue() + FTimespan::FromSeconds((bReverse ? kOutdatedVideoSamplesTolerance : -kOutdatedVideoSamplesTolerance)), bReverse);
 		SET_DWORD_STAT(STAT_MediaUtils_FacadeNumPurgedSubtitleSamples, NumPurged);
 		INC_DWORD_STAT_BY(STAT_MediaUtils_FacadeTotalPurgedSubtitleSamples, NumPurged);
 	}
@@ -2779,46 +2779,7 @@ void FMediaPlayerFacade::ReceiveMediaEvent(EMediaEvent Event)
 		{
 		case	EMediaEvent::Internal_PurgeVideoSamplesHint:
 		{
-			//
-			// Player asks to attempt to purge older samples in the video output queue it maintains
-			// (ask goes via facade as the player does not have accurate timing info)
-			//
-			TSharedPtr<IMediaPlayer, ESPMode::ThreadSafe> CurrentPlayer = Player;
-
-			if (!CurrentPlayer.IsValid())
-			{
-				return;
-			}
-
-			// We only support this for V2 timing players
-			check(CurrentPlayer->GetPlayerFeatureFlag(IMediaPlayer::EFeatureFlag::UsePlaybackTimingV2));
-
-			// Only do this if we do not block on time ranges
-			if (BlockOnRange.IsSet())
-			{
-				// We do not purge as we do not need max perf, but max reliability to actually get certain frames
-				return;
-			}
-
-			float Rate = CurrentRate;
-			if (Rate == 0.0f)
-			{
-				return;
-			}
-
-			// Get current playback time
-			// (note: we have DeltaTime forced to zero -> we just get a single value & we compute relative to "now", noty any game frame start)
-			TRange<FMediaTimeStamp> TimeRange;
-			if (!GetCurrentPlaybackTimeRange(TimeRange, Rate, FTimespan::Zero(), true))
-			{
-				return;
-			}
-
-			bool bReverse = (Rate < 0.0f);
-			uint32 NumPurged = CurrentPlayer->GetSamples().PurgeOutdatedVideoSamples(TimeRange.GetLowerBoundValue() + (bReverse ? kOutdatedVideoSamplesTolerance : -kOutdatedVideoSamplesTolerance), bReverse);
-			SET_DWORD_STAT(STAT_MediaUtils_FacadeNumPurgedVideoSamples, NumPurged);
-			INC_DWORD_STAT_BY(STAT_MediaUtils_FacadeTotalPurgedVideoSamples, NumPurged);
-
+			// This is a no-op now. Case is still here to prevent log spamming in the default handler.
 			break;
 		}
 
