Bringing back video sample purging logic to player facade

diff --git a/Engine/Source/Runtime/MediaUtils/Private/MediaPlayerFacade.cpp b/Engine/Source/Runtime/MediaUtils/Private/MediaPlayerFacade.cpp
--- a/Engine/Source/Runtime/MediaUtils/Private/MediaPlayerFacade.cpp
+++ b/Engine/Source/Runtime/MediaUtils/Private/MediaPlayerFacade.cpp
@@ -79,10 +79,11 @@ DECLARE_DWORD_ACCUMULATOR_STAT(TEXT("MediaPlayerFacade TotalPurgedSubtitleSample
 /* Some constants
 *****************************************************************************/
 
-const double kMaxTimeSinceFrameStart = 0.300; // max seconds we allow between the start of the frame and the player facade timing computations (to catch suspended apps & debugging)
-const double kMaxTimeSinceAudioTimeSampling = 0.250; // max seconds we allow to have passed between the last audio timing sampling and the player facade timing computations (to catch suspended apps & debugging - some platforms do update audio at a farily low rate: hence the big tollerance)
-const double kOutdatedVideoSamplesTolerance = 0.050; // seconds video samples are allowed to be "too old" to stay in the player's output queue despite of calculations indicating they need to go
-const double kOutdatedSubtitleSamplesTolerance = 0.050; // seconds subtitle samples are allowed to be "too old" to stay in the player's output queue despite of calculations indicating they need to go
+static const double kMaxTimeSinceFrameStart = 0.300;			// max seconds we allow between the start of the frame and the player facade timing computations (to catch suspended apps & debugging)
+static const double kMaxTimeSinceAudioTimeSampling = 0.250;		// max seconds we allow to have passed between the last audio timing sampling and the player facade timing computations (to catch suspended apps & debugging - some platforms do update audio at a farily low rate: hence the big tollerance)
+static const double kOutdatedVideoSamplesTolerance = 0.050;		// seconds video samples are allowed to be "too old" to stay in the player's output queue despite of calculations indicating they need to go
+static const double kOutdatedSubtitleSamplesTolerance = 0.050;	// seconds subtitle samples are allowed to be "too old" to stay in the player's output queue despite of calculations indicating they need to go
+static const double kOutdatedSamplePurgeRange = 1.0;			// milliseconds for pseudo DT timespan used with async purging of outdated video samples
 
 /* Local helpers
 *****************************************************************************/
@@ -2779,7 +2780,52 @@ void FMediaPlayerFacade::ReceiveMediaEvent(EMediaEvent Event)
 		{
 		case	EMediaEvent::Internal_PurgeVideoSamplesHint:
 		{
-			// This is a no-op now. Case is still here to prevent log spamming in the default handler.
+			//
+			// Player asks to attempt to purge older samples in the video output queue it maintains
+			// (ask goes via facade as the player does not have accurate timing info)
+			//
+			TSharedPtr<IMediaPlayer, ESPMode::ThreadSafe> CurrentPlayer = Player;
+
+			if (!CurrentPlayer.IsValid())
+			{
+				return;
+			}
+
+			// We only support this for V2 timing players
+			check(CurrentPlayer->GetPlayerFeatureFlag(IMediaPlayer::EFeatureFlag::UsePlaybackTimingV2));
+
+			// Only do this if we do not block on time ranges
+			if (BlockOnRange.IsSet())
+			{
+				// We do not purge as we do not need max perf, but max reliability to actually get certain frames
+				return;
+			}
+
+			float Rate = CurrentRate;
+			if (Rate == 0.0f)
+			{
+				return;
+			}
+
+			// Get current playback time
+			// (Note: the delta time is entirely synthetic - we do not pass zero to avoid an empty range, but we do not look far into the future either
+			//        -> after all: we are mainly focused on purging samples up to the current time
+			//  Remarks:
+			//   - this version does not take any estimations from any frame start into account as this is entirely async to the main thread
+			//   - video streams with no audio content will be played using the UE DeltaTime -> so if that stops, the progress of the video stops!
+			//     -> hence we will not see (other then one initial purge) any purging of samples here!
+			// )
+			TRange<FMediaTimeStamp> TimeRange;
+			if (!GetCurrentPlaybackTimeRange(TimeRange, Rate, FTimespan::FromMilliseconds(kOutdatedSamplePurgeRange), true))
+			{
+				return;
+			}
+
+			bool bReverse = (Rate < 0.0f);
+			uint32 NumPurged = CurrentPlayer->GetSamples().PurgeOutdatedVideoSamples(TimeRange.GetLowerBoundValue() + FTimespan::FromSeconds(bReverse ? kOutdatedVideoSamplesTolerance : -kOutdatedVideoSamplesTolerance), bReverse);
+			SET_DWORD_STAT(STAT_MediaUtils_FacadeNumPurgedVideoSamples, NumPurged);
+			INC_DWORD_STAT_BY(STAT_MediaUtils_FacadeTotalPurgedVideoSamples, NumPurged);
+
 			break;
 		}
 
