Temporary fix for MFW's facade pruging video frames way too eagerly
#rnx

diff --git a/Engine/Source/Runtime/MediaUtils/Private/MediaPlayerFacade.cpp b/Engine/Source/Runtime/MediaUtils/Private/MediaPlayerFacade.cpp
--- a/Engine/Source/Runtime/MediaUtils/Private/MediaPlayerFacade.cpp
+++ b/Engine/Source/Runtime/MediaUtils/Private/MediaPlayerFacade.cpp
@@ -2749,7 +2749,7 @@ void FMediaPlayerFacade::ProcessSubtitleSamples(IMediaSamples& Samples, TRange<F
 	if (CurrentPlayer.IsValid() && CurrentPlayer->GetPlayerFeatureFlag(IMediaPlayer::EFeatureFlag::UsePlaybackTimingV2))
 	{
 		bool bReverse = CurrentRate < 0.0f;
-		uint32 NumPurged = CurrentPlayer->GetSamples().PurgeOutdatedSubtitleSamples(TimeRange.GetLowerBoundValue() + FTimespan::FromSeconds((bReverse ? kOutdatedVideoSamplesTolerance : -kOutdatedVideoSamplesTolerance)), bReverse);
+		uint32 NumPurged = CurrentPlayer->GetSamples().PurgeOutdatedSubtitleSamples(TimeRange.GetLowerBoundValue() + FTimespan::FromSeconds((bReverse ? kOutdatedSubtitleSamplesTolerance : -kOutdatedSubtitleSamplesTolerance)), bReverse);
 		SET_DWORD_STAT(STAT_MediaUtils_FacadeNumPurgedSubtitleSamples, NumPurged);
 		INC_DWORD_STAT_BY(STAT_MediaUtils_FacadeTotalPurgedSubtitleSamples, NumPurged);
 	}
@@ -2780,6 +2780,8 @@ void FMediaPlayerFacade::ReceiveMediaEvent(EMediaEvent Event)
 		{
 		case	EMediaEvent::Internal_PurgeVideoSamplesHint:
 		{
+return; // Disabled for now (fix pending)
+
 			//
 			// Player asks to attempt to purge older samples in the video output queue it maintains
 			// (ask goes via facade as the player does not have accurate timing info)
