Fixed audio sample dropping to consider the sequence index so that future samples will be dropped when looping.

#rnx

diff --git a/Engine/Source/Runtime/Media/Public/IMediaSamples.h b/Engine/Source/Runtime/Media/Public/IMediaSamples.h
--- a/Engine/Source/Runtime/Media/Public/IMediaSamples.h
+++ b/Engine/Source/Runtime/Media/Public/IMediaSamples.h
@@ -37,6 +37,10 @@ public:
 	{
 		return false; // override in child classes, if supported
 	}
+	virtual bool FetchAudio(TRange<FMediaTimeStamp> TimeRange, TSharedPtr<IMediaAudioSample, ESPMode::ThreadSafe>& OutSample)
+	{
+		return false; // override in child classes, if supported
+	}
 
 	/**
 	 * Fetch the next caption sample.
diff --git a/Engine/Source/Runtime/MediaUtils/Private/MediaPlayerFacade.cpp b/Engine/Source/Runtime/MediaUtils/Private/MediaPlayerFacade.cpp
--- a/Engine/Source/Runtime/MediaUtils/Private/MediaPlayerFacade.cpp
+++ b/Engine/Source/Runtime/MediaUtils/Private/MediaPlayerFacade.cpp
@@ -2483,13 +2483,23 @@ void FMediaPlayerFacade::ProcessAudioSamples(IMediaSamples& Samples, TRange<FTim
 			// Do we have video playback?
 			if (HaveVideoPlayback())
 			{
+				TRange<FMediaTimeStamp> TempRange;
 				// We got video and audio, but no audio sink - throw away anything up to video playback time...
 				// (rough estimate, as this is off-gamethread; but better than throwing things out with no throttling at all)
 				{
+					bool bReverse = (CurrentRate < 0.0f);
 					FScopeLock Lock(&LastTimeValuesCS);
-					TimeRange.SetUpperBound(TRangeBound<FTimespan>(CurrentFrameVideoTimeStamp.Time));
+					if (!bReverse)
+					{
+						TempRange.SetUpperBound(CurrentFrameVideoTimeStamp);
+					}
+					else
+					{
+						TempRange.SetLowerBound(CurrentFrameVideoTimeStamp);
+					}
+
 				}
-				while (Samples.FetchAudio(TimeRange, Sample))
+				while (Samples.FetchAudio(TempRange, Sample))
 					;
 			}
 			else
diff --git a/Engine/Source/Runtime/MediaUtils/Private/MediaSamples.cpp b/Engine/Source/Runtime/MediaUtils/Private/MediaSamples.cpp
--- a/Engine/Source/Runtime/MediaUtils/Private/MediaSamples.cpp
+++ b/Engine/Source/Runtime/MediaUtils/Private/MediaSamples.cpp
@@ -92,6 +92,11 @@ bool FMediaSamples::FetchVideo(TRange<FTimespan> TimeRange, TSharedPtr<IMediaTex
 	return FetchSample(VideoSampleQueue, TimeRange, OutSample);
 }
 
+bool FMediaSamples::FetchAudio(TRange<FMediaTimeStamp> TimeRange, TSharedPtr<IMediaAudioSample, ESPMode::ThreadSafe>& OutSample)
+{
+	return FetchSample(AudioSampleQueue, TimeRange, OutSample);
+}
+
 bool FMediaSamples::FetchCaption(TRange<FMediaTimeStamp> TimeRange, TSharedPtr<IMediaOverlaySample, ESPMode::ThreadSafe>& OutSample)
 {
 	return FetchSample(CaptionSampleQueue, TimeRange, OutSample);
diff --git a/Engine/Source/Runtime/MediaUtils/Public/MediaSamples.h b/Engine/Source/Runtime/MediaUtils/Public/MediaSamples.h
--- a/Engine/Source/Runtime/MediaUtils/Public/MediaSamples.h
+++ b/Engine/Source/Runtime/MediaUtils/Public/MediaSamples.h
@@ -150,6 +150,7 @@ public:
 	virtual bool FetchVideo(TRange<FTimespan> TimeRange, TSharedPtr<IMediaTextureSample, ESPMode::ThreadSafe>& OutSample) override;
 	virtual void FlushSamples() override;
 
+	virtual bool FetchAudio(TRange<FMediaTimeStamp> TimeRange, TSharedPtr<IMediaAudioSample, ESPMode::ThreadSafe>& OutSample) override;
 	virtual bool FetchCaption(TRange<FMediaTimeStamp> TimeRange, TSharedPtr<IMediaOverlaySample, ESPMode::ThreadSafe>& OutSample) override;
 	virtual bool FetchSubtitle(TRange<FMediaTimeStamp> TimeRange, TSharedPtr<IMediaOverlaySample, ESPMode::ThreadSafe>& OutSample) override;
 	virtual bool FetchVideo(TRange<FMediaTimeStamp> TimeRange, TSharedPtr<IMediaTextureSample, ESPMode::ThreadSafe>& OutSample) override;
