Fixing up and reactivating MediaFramework's puriging of outdated video frames (now looks at squence IDs and leaves small queues untouched, hence allowing for slow-perf video decodinng to still output frames)

diff --git a/Engine/Source/Runtime/MediaUtils/Private/MediaPlayerFacade.cpp b/Engine/Source/Runtime/MediaUtils/Private/MediaPlayerFacade.cpp
--- a/Engine/Source/Runtime/MediaUtils/Private/MediaPlayerFacade.cpp
+++ b/Engine/Source/Runtime/MediaUtils/Private/MediaPlayerFacade.cpp
@@ -81,9 +81,11 @@ DECLARE_DWORD_ACCUMULATOR_STAT(TEXT("MediaPlayerFacade TotalPurgedSubtitleSample
 
 static const double kMaxTimeSinceFrameStart = 0.300;			// max seconds we allow between the start of the frame and the player facade timing computations (to catch suspended apps & debugging)
 static const double kMaxTimeSinceAudioTimeSampling = 0.250;		// max seconds we allow to have passed between the last audio timing sampling and the player facade timing computations (to catch suspended apps & debugging - some platforms do update audio at a farily low rate: hence the big tollerance)
-static const double kOutdatedVideoSamplesTolerance = 0.050;		// seconds video samples are allowed to be "too old" to stay in the player's output queue despite of calculations indicating they need to go
+static const double kOutdatedVideoSamplesTolerance = 0.080;		// seconds video samples are allowed to be "too old" to stay in the player's output queue despite of calculations indicating they need to go
 static const double kOutdatedSubtitleSamplesTolerance = 0.050;	// seconds subtitle samples are allowed to be "too old" to stay in the player's output queue despite of calculations indicating they need to go
 static const double kOutdatedSamplePurgeRange = 1.0;			// milliseconds for pseudo DT timespan used with async purging of outdated video samples
+static const int32	kMinFramesInVideoQueueToPurge = 3;			// we only consider purging any old frames from the video queue if more than these are present (to not kill a slow playback entirely)
+static const int32	kMinFramesInSubtitleQueueToPurge = 3;		// we only consider purging any old frames from the subtitle queue if more than these are present (to not kill a slow playback entirely)
 
 /* Local helpers
 *****************************************************************************/
@@ -2757,8 +2759,12 @@ void FMediaPlayerFacade::ProcessSubtitleSamples(IMediaSamples& Samples, TRange<F
 	TSharedPtr<IMediaPlayer, ESPMode::ThreadSafe> CurrentPlayer = Player;
 	if (CurrentPlayer.IsValid() && CurrentPlayer->GetPlayerFeatureFlag(IMediaPlayer::EFeatureFlag::UsePlaybackTimingV2))
 	{
-		bool bReverse = CurrentRate < 0.0f;
-		uint32 NumPurged = CurrentPlayer->GetSamples().PurgeOutdatedSubtitleSamples(TimeRange.GetLowerBoundValue() + FTimespan::FromSeconds((bReverse ? kOutdatedSubtitleSamplesTolerance : -kOutdatedSubtitleSamplesTolerance)), bReverse);
+		uint32 NumPurged = 0;
+		if (CurrentPlayer->GetSamples().NumSubtitleSamples() >= kMinFramesInSubtitleQueueToPurge)
+		{
+			bool bReverse = CurrentRate < 0.0f;
+			NumPurged = CurrentPlayer->GetSamples().PurgeOutdatedSubtitleSamples(TimeRange.GetLowerBoundValue() + FTimespan::FromSeconds((bReverse ? kOutdatedSubtitleSamplesTolerance : -kOutdatedSubtitleSamplesTolerance)), bReverse);
+		}
 		SET_DWORD_STAT(STAT_MediaUtils_FacadeNumPurgedSubtitleSamples, NumPurged);
 		INC_DWORD_STAT_BY(STAT_MediaUtils_FacadeTotalPurgedSubtitleSamples, NumPurged);
 	}
@@ -2789,8 +2795,6 @@ void FMediaPlayerFacade::ReceiveMediaEvent(EMediaEvent Event)
 		{
 		case	EMediaEvent::Internal_PurgeVideoSamplesHint:
 		{
-return; // Disabled for now (fix pending)
-
 			//
 			// Player asks to attempt to purge older samples in the video output queue it maintains
 			// (ask goes via facade as the player does not have accurate timing info)
@@ -2832,8 +2836,14 @@ return; // Disabled for now (fix pending)
 				return;
 			}
 
-			bool bReverse = (Rate < 0.0f);
-			uint32 NumPurged = CurrentPlayer->GetSamples().PurgeOutdatedVideoSamples(TimeRange.GetLowerBoundValue() + FTimespan::FromSeconds(bReverse ? kOutdatedVideoSamplesTolerance : -kOutdatedVideoSamplesTolerance), bReverse);
+			uint32 NumPurged = 0;
+			// Don't purge frames if the queue is small (to avoid purging if players deliver frames late persistently)
+			if (CurrentPlayer->GetSamples().NumVideoSamples() >= kMinFramesInVideoQueueToPurge)
+			{
+				bool bReverse = (Rate < 0.0f);
+				NumPurged = CurrentPlayer->GetSamples().PurgeOutdatedVideoSamples(TimeRange.GetLowerBoundValue() + FTimespan::FromSeconds(bReverse ? kOutdatedVideoSamplesTolerance : -kOutdatedVideoSamplesTolerance), bReverse);
+			}
+
 			SET_DWORD_STAT(STAT_MediaUtils_FacadeNumPurgedVideoSamples, NumPurged);
 			INC_DWORD_STAT_BY(STAT_MediaUtils_FacadeTotalPurgedVideoSamples, NumPurged);
 
diff --git a/Engine/Source/Runtime/MediaUtils/Public/MediaSampleQueue.h b/Engine/Source/Runtime/MediaUtils/Public/MediaSampleQueue.h
--- a/Engine/Source/Runtime/MediaUtils/Public/MediaSampleQueue.h
+++ b/Engine/Source/Runtime/MediaUtils/Public/MediaSampleQueue.h
@@ -252,7 +252,7 @@ public:
 				for (; Idx < Num; ++Idx)
 				{
 					const TSharedPtr<SampleType, ESPMode::ThreadSafe> & Sample = Samples[Idx];
-					if ((Sample->GetTime().Time + Sample->GetDuration()) > ReferenceTime.Time)
+					if ((Sample->GetTime() + Sample->GetDuration()) > ReferenceTime)
 					{
 						break;
 					}
@@ -263,7 +263,7 @@ public:
 				for (; Idx < Num; ++Idx)
 				{
 					const TSharedPtr<SampleType, ESPMode::ThreadSafe> & Sample = Samples[Idx];
-					if ((Sample->GetTime().Time - Sample->GetDuration()) < ReferenceTime.Time)
+					if ((Sample->GetTime() - Sample->GetDuration()) < ReferenceTime)
 					{
 						break;
 					}
