Moving check for audio-in-media to disable blocked/synced playbackin mediaFramework into player-version-independent code
#rnx

diff --git a/Engine/Source/Runtime/MediaUtils/Private/MediaPlayerFacade.cpp b/Engine/Source/Runtime/MediaUtils/Private/MediaPlayerFacade.cpp
--- a/Engine/Source/Runtime/MediaUtils/Private/MediaPlayerFacade.cpp
+++ b/Engine/Source/Runtime/MediaUtils/Private/MediaPlayerFacade.cpp
@@ -1483,9 +1483,9 @@ bool FMediaPlayerFacade::BlockOnFetch() const
 
 	const TRange<FMediaTimeStamp>& BR = BlockOnRange.GetRange();
 
-	if (BR.IsEmpty() || !Player->GetControls().CanControl(EMediaControl::BlockOnFetch) || BlockOnRangeDisabled)
+	if (BR.IsEmpty() || !Player->GetControls().CanControl(EMediaControl::BlockOnFetch) || BlockOnRangeDisabled || bHaveActiveAudio)
 	{
-		return false; // no blocking requested / not supported
+		return false; // no blocking requested / not supported / audio present
 	}
 
 	if (Player->GetPlayerFeatureFlag(IMediaPlayer::EFeatureFlag::UsePlaybackTimingV2))
@@ -1494,12 +1494,6 @@ bool FMediaPlayerFacade::BlockOnFetch() const
 		// V2 blocking logic
 		//
 
-		// If we have any active audio playback we skip any blocking
-		if (HaveAudioPlayback())
-		{
-			return false;
-		}
-
 		float Rate = GetUnpausedRate();
 
 		// If the current sample "out there" is actually overlapping with the current block, we might be good with no new sample
