Updated MFW's PlayerFacade with reworked handling for playback control and timing for "v2 timing" players
- New way of handling seqeunce index values
- Reduced / reworked flush behavior on seeks
- Adjusted ImageMediaPlayer, Electra & to some extend WMFMedia to this new setup
- Debugged and fixed up blocked  / synced playback using the changes to the sequence index handling
- Added new display oriented timestamp retrieval method so GUIs can display PTS in a more user related way
- Enforced that "v2 timing" players return a first frame to the facade on opening or seeking no matter if paused or not

diff --git a/Engine/Plugins/Experimental/TextureMediaPlayer/Source/TextureMediaPlayer/Private/TextureMediaPlayer.cpp b/Engine/Plugins/Experimental/TextureMediaPlayer/Source/TextureMediaPlayer/Private/TextureMediaPlayer.cpp
--- a/Engine/Plugins/Experimental/TextureMediaPlayer/Source/TextureMediaPlayer/Private/TextureMediaPlayer.cpp
+++ b/Engine/Plugins/Experimental/TextureMediaPlayer/Source/TextureMediaPlayer/Private/TextureMediaPlayer.cpp
@@ -431,7 +431,7 @@ void FTextureMediaPlayer::OnFrame(const TArray<uint8>& TextureBuffer, FIntPoint
 			// NOTE: We cannot use Internal_PurgeVideoSamplesHint here because if the game loop is not
 			//       being ticked the time it uses for comparison will not have advanced and nothing
 			//       will get purged.
-			MediaSamples->PurgeOutdatedVideoSamples(FMediaTimeStamp(FTimespan::MaxValue()), false);
+			MediaSamples->PurgeOutdatedVideoSamples(FMediaTimeStamp(FTimespan::MaxValue()), false, FTimespan::Zero());
 			EventSink.ReceiveMediaEvent(EMediaEvent::Internal_ResetForDiscontinuity);
 		}
 
@@ -484,7 +484,7 @@ void FTextureMediaPlayer::OnFrame(FTextureRHIRef TextureRHIRef, TRefCountPtr<ID3
 			// NOTE: We cannot use Internal_PurgeVideoSamplesHint here because if the game loop is not
 			//       being ticked the time it uses for comparison will not have advanced and nothing
 			//       will get purged.
-			MediaSamples->PurgeOutdatedVideoSamples(FMediaTimeStamp(FTimespan::MaxValue()), false);
+			MediaSamples->PurgeOutdatedVideoSamples(FMediaTimeStamp(FTimespan::MaxValue()), false, FTimespan::Zero());
 			EventSink.ReceiveMediaEvent(EMediaEvent::Internal_ResetForDiscontinuity);
 		}
 
diff --git a/Engine/Plugins/Media/ElectraPlayer/Source/ElectraPlayerPlugin/Private/ElectraPlayerPlugin.cpp b/Engine/Plugins/Media/ElectraPlayer/Source/ElectraPlayerPlugin/Private/ElectraPlayerPlugin.cpp
--- a/Engine/Plugins/Media/ElectraPlayer/Source/ElectraPlayerPlugin/Private/ElectraPlayerPlugin.cpp
+++ b/Engine/Plugins/Media/ElectraPlayer/Source/ElectraPlayerPlugin/Private/ElectraPlayerPlugin.cpp
@@ -16,22 +16,10 @@
 #include "MediaSubtitleDecoderOutput.h"
 #include "MediaMetaDataDecoderOutput.h"
 
-#include "HAL/IConsoleManager.h"
-
 using namespace Electra;
 
 //-----------------------------------------------------------------------------
 
-static int GElectraEnableBlockOnFetch = 0;							// Temporarily disable Electra's support for BlockOnFetch by default
-static FAutoConsoleVariableRef CVarElectraEnableBlockOnFetch(
-	TEXT("ElectraPlayer.EnableBlockOnFetch"),
-	GElectraEnableBlockOnFetch,
-	TEXT("Whether 'BlockOnFetch' support is enabled.\n")
-	TEXT(" 0: disabled (default)\n")
-	TEXT(" 1: enabled"));
-
-//-----------------------------------------------------------------------------
-
 FElectraPlayerPlugin::FElectraPlayerPlugin()
 {
 	// Make sure a few assumptions are correct...
@@ -460,6 +448,7 @@ void FElectraPlayerPlugin::FPlayerAdapterDelegate::PresentAudioFrame(const IAudi
 	}
 }
 
+
 void FElectraPlayerPlugin::FPlayerAdapterDelegate::PresentSubtitleSample(const ISubtitleDecoderOutputPtr& InSubtitleSample)
 {
 	TSharedPtr<FElectraPlayerPlugin, ESPMode::ThreadSafe> PinnedHost = Host.Pin();
@@ -475,6 +464,7 @@ void FElectraPlayerPlugin::FPlayerAdapterDelegate::PresentSubtitleSample(const I
 	}
 }
 
+
 void FElectraPlayerPlugin::FPlayerAdapterDelegate::PresentMetadataSample(const IMetaDataDecoderOutputPtr& InMetadataFrame)
 {
 	TSharedPtr<FElectraPlayerPlugin, ESPMode::ThreadSafe> PinnedHost = Host.Pin();
@@ -847,11 +837,7 @@ bool FElectraPlayerPlugin::CanControl(EMediaControl Control) const
 	EMediaState CurrentState = GetState();
 	if (Control == EMediaControl::BlockOnFetch)
 	{
-		if (GElectraEnableBlockOnFetch == 0)
-		{
-			return false;
-		}
-		return CurrentState == EMediaState::Playing;
+		return CurrentState == EMediaState::Paused || CurrentState == EMediaState::Playing;
 	}
 	else if (Control == EMediaControl::Pause)
 	{
diff --git a/Engine/Plugins/Media/ElectraPlayer/Source/ElectraPlayerPlugin/Private/ElectraPlayerPlugin.h b/Engine/Plugins/Media/ElectraPlayer/Source/ElectraPlayerPlugin/Private/ElectraPlayerPlugin.h
--- a/Engine/Plugins/Media/ElectraPlayer/Source/ElectraPlayerPlugin/Private/ElectraPlayerPlugin.h
+++ b/Engine/Plugins/Media/ElectraPlayer/Source/ElectraPlayerPlugin/Private/ElectraPlayerPlugin.h
@@ -97,14 +97,6 @@ public:
 	bool Open(const TSharedRef<FArchive, ESPMode::ThreadSafe>& Archive, const FString& OriginalUrl, const IMediaOptions* Options) override;
 	void TickInput(FTimespan DeltaTime, FTimespan Timecode) override;
 	void SetLastAudioRenderedSampleTime(FTimespan SampleTime) override;
-	bool FlushOnSeekStarted() const override
-	{
-		return true;
-	}
-	bool FlushOnSeekCompleted() const override
-	{
-		return false;
-	}
 
 	TSharedPtr<TMap<FString, TArray<TUniquePtr<IMediaPlayer::IMetadataItem>>>, ESPMode::ThreadSafe> GetMediaMetadata() const override;
 
diff --git a/Engine/Plugins/Media/ElectraPlayer/Source/ElectraPlayerRuntime/Private/ElectraPlayer.cpp b/Engine/Plugins/Media/ElectraPlayer/Source/ElectraPlayerRuntime/Private/ElectraPlayer.cpp
--- a/Engine/Plugins/Media/ElectraPlayer/Source/ElectraPlayerRuntime/Private/ElectraPlayer.cpp
+++ b/Engine/Plugins/Media/ElectraPlayer/Source/ElectraPlayerRuntime/Private/ElectraPlayer.cpp
@@ -1082,6 +1082,10 @@ void FElectraPlayer::TriggerFirstSeekIfNecessary()
 
 		// Set up the initial playback position
 		IAdaptiveStreamingPlayer::FSeekParam playParam;
+
+		// Whatever the initial seek target: we will ignore it for sequence index updates
+		playParam.bIgnoreForSequenceIndex = true;
+
 		// First we look at any potential time offset specified in the playstart options.
 		if (PlaystartOptions.TimeOffset.IsSet())
 		{
diff --git a/Engine/Plugins/Media/ElectraPlayer/Source/ElectraPlayerRuntime/Private/Runtime/Player/AdaptiveStreamingPlayer.cpp b/Engine/Plugins/Media/ElectraPlayer/Source/ElectraPlayerRuntime/Private/Runtime/Player/AdaptiveStreamingPlayer.cpp
--- a/Engine/Plugins/Media/ElectraPlayer/Source/ElectraPlayerRuntime/Private/Runtime/Player/AdaptiveStreamingPlayer.cpp
+++ b/Engine/Plugins/Media/ElectraPlayer/Source/ElectraPlayerRuntime/Private/Runtime/Player/AdaptiveStreamingPlayer.cpp
@@ -628,6 +628,10 @@ void FAdaptiveStreamingPlayer::SeekTo(const FSeekParam& NewPosition)
 	FScopeLock lock(&SeekVars.Lock);
 	SeekVars.PendingRequest = NewPosition;
 	SeekVars.PlayrangeOnRequest = NewPlayRange;
+	if (!NewPosition.bIgnoreForSequenceIndex.Get(false))
+	{
+		++SeekVars.NumSeekToCallsSinceLastSeen;
+	}
 	WorkerThread.TriggerSharedWorkerThread();
 }
 
@@ -1660,9 +1664,16 @@ void FAdaptiveStreamingPlayer::HandleSeeking()
 
 		// Not playing on the Live edge. This may get set to true later on handling the start segment.
 		PlaybackState.SetShouldPlayOnLiveEdge(false);
-		// On a user-induced seek the sequence index is supposed to reset to 0.
+		// On a user-induced seek the sequence index is supposed to be increased
 		// We need to reset this even when we would not actually issue the seek.
-		CurrentPlaybackSequenceState.SequenceIndex = 0;
+
+		if (!SeekVars.PendingRequest.GetValue().bIgnoreForSequenceIndex.Get(false))
+		{
+			CurrentPlaybackSequenceState.PrimaryIndex += SeekVars.NumSeekToCallsSinceLastSeen;
+			CurrentPlaybackSequenceState.SecondaryIndex = 0;
+			SeekVars.NumSeekToCallsSinceLastSeen = 0;
+		}
+
 		// And since it is a seek on purpose the loop counter is reset as well.
 		FInternalLoopState LoopStateNow;
 		PlaybackState.GetLoopState(LoopStateNow);
@@ -2679,7 +2690,7 @@ void FAdaptiveStreamingPlayer::InternalHandlePendingStartRequest(const FTimeValu
 							else
 							{
 								FPlayerSequenceState NextPlaybackState = CurrentPlaybackSequenceState;
-								++NextPlaybackState.SequenceIndex;
+								++NextPlaybackState.SecondaryIndex;
 								Result = InitialPlayPeriod->GetLoopingSegment(FirstSegmentRequest, NextPlaybackState, StartAt, PendingStartRequest->SearchType);
 								if (Result.GetType() == IManifest::FResult::EType::Found)
 								{
diff --git a/Engine/Plugins/Media/ElectraPlayer/Source/ElectraPlayerRuntime/Private/Runtime/Player/AdaptiveStreamingPlayer.h b/Engine/Plugins/Media/ElectraPlayer/Source/ElectraPlayerRuntime/Private/Runtime/Player/AdaptiveStreamingPlayer.h
--- a/Engine/Plugins/Media/ElectraPlayer/Source/ElectraPlayerRuntime/Private/Runtime/Player/AdaptiveStreamingPlayer.h
+++ b/Engine/Plugins/Media/ElectraPlayer/Source/ElectraPlayerRuntime/Private/Runtime/Player/AdaptiveStreamingPlayer.h
@@ -130,6 +130,8 @@ public:
 		TOptional<bool> bOptimizeForScrubbing;
 		// Allowed distance to last performed seek to save a redundant new seek.
 		TOptional<double> DistanceThreshold;
+		// Ignore the seek for sequence index updates
+		TOptional<bool> bIgnoreForSequenceIndex;
 	};
 	
 	/**
diff --git a/Engine/Plugins/Media/ElectraPlayer/Source/ElectraPlayerRuntime/Private/Runtime/Player/AdaptiveStreamingPlayerInternal.h b/Engine/Plugins/Media/ElectraPlayer/Source/ElectraPlayerRuntime/Private/Runtime/Player/AdaptiveStreamingPlayerInternal.h
--- a/Engine/Plugins/Media/ElectraPlayer/Source/ElectraPlayerRuntime/Private/Runtime/Player/AdaptiveStreamingPlayerInternal.h
+++ b/Engine/Plugins/Media/ElectraPlayer/Source/ElectraPlayerRuntime/Private/Runtime/Player/AdaptiveStreamingPlayerInternal.h
@@ -1776,7 +1776,9 @@ private:
 		void Reset()
 		{
 			FScopeLock lock(&Lock);
-			PendingRequest.Reset();
+			// Note: Do _not_ reset the pending request created by a user induced seek, or it may get lost if
+			//       triggered while already processing a seek due to the asynchronous processing of the request.
+			//PendingRequest.Reset();
 			ActiveRequest.Reset();
 			LastFinishedRequest.Reset();
 			PlayrangeOnRequest.Reset();
@@ -1798,6 +1800,7 @@ private:
 		bool	bForScrubbing = false;
 		bool	bScrubPrerollDone = false;
 		bool	bIsPlayStart = true;
+		int32	NumSeekToCallsSinceLastSeen = 0;
 
 		TOptional<FSeekParam> ActiveRequest;
 		TOptional<FSeekParam> LastFinishedRequest;
diff --git a/Engine/Plugins/Media/ElectraPlayer/Source/ElectraPlayerRuntime/Private/Runtime/Player/DASH/ManifestDASH.cpp b/Engine/Plugins/Media/ElectraPlayer/Source/ElectraPlayerRuntime/Private/Runtime/Player/DASH/ManifestDASH.cpp
--- a/Engine/Plugins/Media/ElectraPlayer/Source/ElectraPlayerRuntime/Private/Runtime/Player/DASH/ManifestDASH.cpp
+++ b/Engine/Plugins/Media/ElectraPlayer/Source/ElectraPlayerRuntime/Private/Runtime/Player/DASH/ManifestDASH.cpp
@@ -1111,7 +1111,7 @@ IManifest::FResult FDASHPlayPeriod::GetStartingSegment(TSharedPtrTS<IStreamSegme
 	// dependent streams. This is a special case for playback start.
 	TSharedPtrTS<FStreamSegmentRequestFMP4DASH> StartSegmentRequest = MakeSharedTS<FStreamSegmentRequestFMP4DASH>();
 	StartSegmentRequest->bIsInitialStartRequest = true;
-	StartSegmentRequest->TimestampSequenceIndex = InSequenceState.SequenceIndex;
+	StartSegmentRequest->TimestampSequenceIndex = InSequenceState.GetSequenceIndex();
 
 	struct FSelectedStream
 	{
@@ -1192,7 +1192,7 @@ IManifest::FResult FDASHPlayPeriod::GetStartingSegment(TSharedPtrTS<IStreamSegme
 				}
 				SegmentRequest->Segment = MoveTemp(SegmentInfo);
 				SegmentRequest->bIsEOSSegment = true;
-				SegmentRequest->TimestampSequenceIndex = InSequenceState.SequenceIndex;
+				SegmentRequest->TimestampSequenceIndex = InSequenceState.GetSequenceIndex();
 				SetupCommonSegmentRequestInfos(SegmentRequest);
 				StartSegmentRequest->DependentStreams.Emplace(MoveTemp(SegmentRequest));
 				bAnyStreamAtEOS = true;
@@ -1251,7 +1251,7 @@ IManifest::FResult FDASHPlayPeriod::GetStartingSegment(TSharedPtrTS<IStreamSegme
 					SegmentRequest->bInsertFillerData = true;
 				}
 				SegmentRequest->Segment = MoveTemp(SegmentInfo);
-				SegmentRequest->TimestampSequenceIndex = InSequenceState.SequenceIndex;
+				SegmentRequest->TimestampSequenceIndex = InSequenceState.GetSequenceIndex();
 
 				// The start segment request needs to be able to return a valid first PTS which is what the player sets
 				// the playback position to. If not valid yet update it with the current stream values.
@@ -1312,7 +1312,7 @@ IManifest::FResult FDASHPlayPeriod::GetContinuationSegment(TSharedPtrTS<IStreamS
 	auto DummyReq = MakeSharedTS<FStreamSegmentRequestFMP4DASH>();
 	DummyReq->StreamType = StreamType;
 	DummyReq->PeriodStart = StartPosition.Time;
-	DummyReq->TimestampSequenceIndex = SequenceState.SequenceIndex;
+	DummyReq->TimestampSequenceIndex = SequenceState.GetSequenceIndex();
 	return GetNextOrRetrySegment(OutSegment, bNeedRemoteElement, DummyReq, ENextSegType::SamePeriodStartOver, StartPosition.Options);
 }
 
diff --git a/Engine/Plugins/Media/ElectraPlayer/Source/ElectraPlayerRuntime/Private/Runtime/Player/HLS/ManifestHLS.cpp b/Engine/Plugins/Media/ElectraPlayer/Source/ElectraPlayerRuntime/Private/Runtime/Player/HLS/ManifestHLS.cpp
--- a/Engine/Plugins/Media/ElectraPlayer/Source/ElectraPlayerRuntime/Private/Runtime/Player/HLS/ManifestHLS.cpp
+++ b/Engine/Plugins/Media/ElectraPlayer/Source/ElectraPlayerRuntime/Private/Runtime/Player/HLS/ManifestHLS.cpp
@@ -876,7 +876,7 @@ IManifest::FResult FPlayPeriodHLS::GetStartingSegment(TSharedPtrTS<IStreamSegmen
 		SearchParam.Time = StartPosition.Time;
 		SearchParam.LastPTS = PlayRangeEnd;
 		SearchParam.bFrameAccurateSearch = bFrameAccurateSearch;
-		SearchParam.TimestampSequenceIndex = InSequenceState.SequenceIndex;
+		SearchParam.TimestampSequenceIndex = InSequenceState.GetSequenceIndex();
 
 
 		// Do we have both video and audio?
diff --git a/Engine/Plugins/Media/ElectraPlayer/Source/ElectraPlayerRuntime/Private/Runtime/Player/mp4/ManifestMP4.cpp b/Engine/Plugins/Media/ElectraPlayer/Source/ElectraPlayerRuntime/Private/Runtime/Player/mp4/ManifestMP4.cpp
--- a/Engine/Plugins/Media/ElectraPlayer/Source/ElectraPlayerRuntime/Private/Runtime/Player/mp4/ManifestMP4.cpp
+++ b/Engine/Plugins/Media/ElectraPlayer/Source/ElectraPlayerRuntime/Private/Runtime/Player/mp4/ManifestMP4.cpp
@@ -952,7 +952,7 @@ IManifest::FResult FManifestMP4Internal::FTimelineAssetMP4::GetStartingSegment(T
 					// Set the PTS range for which to present samples.
 					req->EarliestPTS = bFrameAccurateSearch ? StartPosition.Time : firstTimestamp;
 					req->LastPTS = PlayRangeEnd;
-					req->TimestampSequenceIndex = InSequenceState.SequenceIndex;
+					req->TimestampSequenceIndex = InSequenceState.GetSequenceIndex();
 
 					// Add dependent stream types
 					if (AudioAdaptationSets.Num())
@@ -994,7 +994,7 @@ IManifest::FResult FManifestMP4Internal::FTimelineAssetMP4::GetStartingSegment(T
 						// Set the PTS range for which to present samples.
 						req->EarliestPTS = bFrameAccurateSearch ? StartPosition.Time : firstTimestamp;
 						req->LastPTS = PlayRangeEnd;
-						req->TimestampSequenceIndex = InSequenceState.SequenceIndex;
+						req->TimestampSequenceIndex = InSequenceState.GetSequenceIndex();
 						return IManifest::FResult(IManifest::FResult::EType::Found);
 					}
 				}
@@ -1095,7 +1095,7 @@ IManifest::FResult FManifestMP4Internal::FTimelineAssetMP4::GetStartingSegment(T
 					// Set the PTS range for which to present samples.
 					req->EarliestPTS = bFrameAccurateSearch ? StartPosition.Time : firstTimestamp;
 					req->LastPTS = PlayRangeEnd;
-					req->TimestampSequenceIndex = InSequenceState.SequenceIndex;
+					req->TimestampSequenceIndex = InSequenceState.GetSequenceIndex();
 
 					// Add dependent stream types.
 					// In case the video stream is shorter than audio we still need to add it as a dependent stream
@@ -1136,7 +1136,7 @@ IManifest::FResult FManifestMP4Internal::FTimelineAssetMP4::GetStartingSegment(T
 					// Set the PTS range for which to present samples.
 					req->EarliestPTS = bFrameAccurateSearch ? StartPosition.Time : firstTimestamp;
 					req->LastPTS = PlayRangeEnd;
-					req->TimestampSequenceIndex = InSequenceState.SequenceIndex;
+					req->TimestampSequenceIndex = InSequenceState.GetSequenceIndex();
 					// But if there is a video track we add it as a dependent stream that is also at EOS.
 					if (VideoAdaptationSets.Num())
 					{
@@ -1178,7 +1178,7 @@ IManifest::FResult FManifestMP4Internal::FTimelineAssetMP4::GetNextSegment(TShar
 			FPlayStartPosition dummyPos;
 			FPlayerSequenceState seqState;
 			dummyPos.Options = Options;
-			seqState.SequenceIndex = Request->TimestampSequenceIndex;
+			seqState.SetSequenceIndex(Request->TimestampSequenceIndex);
 			IManifest::FResult res = GetStartingSegment(OutSegment, seqState, dummyPos, ESearchType::Same, Request->FileEndOffset + 1);
 			if (res.GetType() == IManifest::FResult::EType::Found)
 			{
@@ -1205,7 +1205,7 @@ IManifest::FResult FManifestMP4Internal::FTimelineAssetMP4::GetRetrySegment(TSha
 		FPlayStartPosition dummyPos;
 		FPlayerSequenceState seqState;
 		dummyPos.Options = Options;
-		seqState.SequenceIndex = Request->TimestampSequenceIndex;
+		seqState.SetSequenceIndex(Request->TimestampSequenceIndex);
 		IManifest::FResult res = GetStartingSegment(OutSegment, seqState, dummyPos, ESearchType::Same, Request->CurrentIteratorBytePos);
 		if (res.GetType() == IManifest::FResult::EType::Found)
 		{
diff --git a/Engine/Plugins/Media/ElectraPlayer/Source/ElectraPlayerRuntime/Private/Runtime/StreamTypes.h b/Engine/Plugins/Media/ElectraPlayer/Source/ElectraPlayerRuntime/Private/Runtime/StreamTypes.h
--- a/Engine/Plugins/Media/ElectraPlayer/Source/ElectraPlayerRuntime/Private/Runtime/StreamTypes.h
+++ b/Engine/Plugins/Media/ElectraPlayer/Source/ElectraPlayerRuntime/Private/Runtime/StreamTypes.h
@@ -805,9 +805,22 @@ namespace Electra
 		}
 		void Reset()
 		{
-			SequenceIndex = 0;
+			PrimaryIndex = 0;
+			SecondaryIndex = 0;
 		}
-		int64	SequenceIndex;
+		int64 GetSequenceIndex() const
+		{
+			return (int64(PrimaryIndex) << 32) + int64(SecondaryIndex);
+		}
+		void SetSequenceIndex(int64 SequenceIndex)
+		{
+			check(SequenceIndex >= 0);	// should be the case in Electra at all times
+			PrimaryIndex = int32(SequenceIndex >> 32);
+			SecondaryIndex = int32(SequenceIndex);
+		}
+
+		int32	PrimaryIndex;
+		int32	SecondaryIndex;
 	};
 
 
diff --git a/Engine/Plugins/Media/ImgMedia/Source/ImgMedia/Internal/Loader/ImgMediaLoader.h b/Engine/Plugins/Media/ImgMedia/Source/ImgMedia/Internal/Loader/ImgMediaLoader.h
--- a/Engine/Plugins/Media/ImgMedia/Source/ImgMedia/Internal/Loader/ImgMediaLoader.h
+++ b/Engine/Plugins/Media/ImgMedia/Source/ImgMedia/Internal/Loader/ImgMediaLoader.h
@@ -149,6 +149,11 @@ public:
 		return Info;
 	}
 
+	/**
+	 * Notify loader out a seek
+	 */
+	void Seek(const FMediaTimeStamp& SeekTarget);
+
 	/**
 	 * Tries to get the best sample for a given time range.
 	 *
@@ -156,7 +161,7 @@ public:
 	 * @param OutSample will be filled in with the sample if found.
 	 * @param bIsLoopingEnabled True if we can loop.
 	 * @param PlayRate How fast we are playing.
-	 * @return True if successful.
+	 * @return Return code indicating status (ok, unsupported, no sample)
 	 */
 	IMediaSamples::EFetchBestSampleResult FetchBestVideoSampleForTimeRange(const TRange<FMediaTimeStamp>& TimeRange, TSharedPtr<IMediaTextureSample, ESPMode::ThreadSafe>& OutSample, bool bIsLoopingEnabled, float PlayRate);
 
@@ -171,6 +176,11 @@ public:
 	 */
 	bool PeekVideoSampleTime(FMediaTimeStamp& TimeStamp, bool bIsLoopingEnabled, float PlayRate, const FTimespan& CurrentTime);
 
+	/**
+	 * Discard any samples in the given range from the output "queue" (logically)
+	 */
+	bool DiscardVideoSamples(const TRange<FMediaTimeStamp>& TimeRange, bool bIsLoopingEnabled, float PlayRate);
+
 	/**
 	 * Get the time ranges of frames that are pending to be loaded.
 	 *
@@ -372,6 +382,12 @@ public:
 	 */
 	void HandlePause();
 
+	/** Query if time represents last first frame or before */
+	bool IsFrameFirst(const FTimespan& TimeStamp) const;
+
+	/** Query if time represents last last frame or later */
+	bool IsFrameLast(const FTimespan& TimeStamp) const;
+
 protected:
 
 	/**
@@ -636,7 +652,9 @@ private:
 	struct
 	{
 		int32 LastFrameIndex;
-		uint64 CurrentSequenceIndex;
+		FMediaTimeStamp LastTimeStamp;
+		FTimespan LastDuration;
+		int32 LoopIndex;
 	} QueuedSampleFetch;
 
 private:
diff --git a/Engine/Plugins/Media/ImgMedia/Source/ImgMedia/Internal/Player/ImgMediaPlayer.h b/Engine/Plugins/Media/ImgMedia/Source/ImgMedia/Internal/Player/ImgMediaPlayer.h
--- a/Engine/Plugins/Media/ImgMedia/Source/ImgMedia/Internal/Player/ImgMediaPlayer.h
+++ b/Engine/Plugins/Media/ImgMedia/Source/ImgMedia/Internal/Player/ImgMediaPlayer.h
@@ -110,6 +110,7 @@ protected:
 
 	virtual EFetchBestSampleResult FetchBestVideoSampleForTimeRange(const TRange<FMediaTimeStamp> & TimeRange, TSharedPtr<IMediaTextureSample, ESPMode::ThreadSafe>& OutSample, bool bReverse) override;
 	virtual bool PeekVideoSampleTime(FMediaTimeStamp & TimeStamp) override;
+	virtual bool DiscardVideoSamples(const TRange<FMediaTimeStamp>& TimeRange, bool bReverse) override;
 
 protected:
 
@@ -135,12 +136,18 @@ private:
 	/** The current playback rate. */
 	float CurrentRate;
 
+	/** The last playback rate not zero (zero if never played). */
+	float LastNonZeroRate;
+
 	/** The player's current state. */
 	EMediaState CurrentState;
 
 	/** The current time of the playback. */
 	FTimespan CurrentTime;
 
+	/** Current seek index portion of sequence index */
+	int32 CurrentSeekIndex;
+
 	/** The URL of the currently opened media. */
 	FString CurrentUrl;
 
diff --git a/Engine/Plugins/Media/ImgMedia/Source/ImgMedia/Private/Loader/ImgMediaLoader.cpp b/Engine/Plugins/Media/ImgMedia/Source/ImgMedia/Private/Loader/ImgMediaLoader.cpp
--- a/Engine/Plugins/Media/ImgMedia/Source/ImgMedia/Private/Loader/ImgMediaLoader.cpp
+++ b/Engine/Plugins/Media/ImgMedia/Source/ImgMedia/Private/Loader/ImgMediaLoader.cpp
@@ -28,6 +28,8 @@
 	#include "Readers/ExrImgMediaReader.h"
 #endif
 
+#define ALLOW_CHEATED_FRAMES 0		// if set to 1 the loader is allowed to return "close by" frame data with a fake PTS to optimize for quick delivery IF there is no blocked range is set
+									// ^^^ DISABLED as this can lead to entirely incorrect frames being displayed when seeking in paused state
 
 /** Time spent loading a new image sequence. */
 DECLARE_CYCLE_STAT(TEXT("ImgMedia Loader Load Sequence"), STAT_ImgMedia_LoaderLoadSequence, STATGROUP_Media);
@@ -343,8 +345,8 @@ const FString& FImgMediaLoader::GetImagePath(int32 FrameNumber, int32 MipLevel)
 void FImgMediaLoader::ResetFetchLogic()
 {
 	QueuedSampleFetch.LastFrameIndex = INDEX_NONE;
-	// note: we can reset the sequence index here as this will be called when MFW does flush any queues - so we can start from scratch with no issues
-	QueuedSampleFetch.CurrentSequenceIndex = 0;
+	QueuedSampleFetch.LastTimeStamp.Invalidate();
+	QueuedSampleFetch.LoopIndex = 0;
 
 	SkipFramesCounter = 0;
 	SkipFramesLevel = 0;
@@ -409,15 +411,95 @@ const TSharedPtr<FImgMediaFrame, ESPMode::ThreadSafe>* FImgMediaLoader::GetFrame
 }
 
 
-IMediaSamples::EFetchBestSampleResult FImgMediaLoader::FetchBestVideoSampleForTimeRange(const TRange<FMediaTimeStamp>& TimeRange, TSharedPtr<IMediaTextureSample, ESPMode::ThreadSafe>& OutSample, bool bIsLoopingEnabled, float PlayRate)
+void FImgMediaLoader::Seek(const FMediaTimeStamp& SeekTarget)
 {
-	if (IsInitialized() && TimeRange.HasLowerBound() && TimeRange.HasUpperBound())
+	ResetFetchLogic();
+	QueuedSampleFetch.LastTimeStamp = SeekTarget;
+	QueuedSampleFetch.LastDuration = FTimespan::Zero();
+	QueuedSampleFetch.LoopIndex = 0;
+	QueuedSampleFetch.LastFrameIndex = INDEX_NONE;
+
+	QueuedSampleFetch.LastTimeStamp.Time = FrameNumberToTime(TimeToFrameNumber(QueuedSampleFetch.LastTimeStamp.Time));
+}
+
+
+bool FImgMediaLoader::DiscardVideoSamples(const TRange<FMediaTimeStamp>& TimeRange, bool bIsLoopingEnabled, float PlayRate)
+{
+	bool bOk = false;
+	if (PlayRate >= 0.0f)
+	{
+		if ((QueuedSampleFetch.LastTimeStamp + QueuedSampleFetch.LastDuration) < TimeRange.GetUpperBoundValue())
+		{
+			QueuedSampleFetch.LastTimeStamp = TimeRange.GetUpperBoundValue();
+			bOk = true;
+		}
+	}
+	else
+	{
+		if ((QueuedSampleFetch.LastTimeStamp - QueuedSampleFetch.LastDuration) > TimeRange.GetLowerBoundValue())
+		{
+			QueuedSampleFetch.LastTimeStamp = TimeRange.GetLowerBoundValue();
+			bOk = true;
+		}
+	}
+
+	if (bOk)
 	{
+		QueuedSampleFetch.LastDuration = FTimespan::Zero();
+		QueuedSampleFetch.LastTimeStamp.Time = FrameNumberToTime(TimeToFrameNumber(QueuedSampleFetch.LastTimeStamp.Time));
+		QueuedSampleFetch.LastFrameIndex = INDEX_NONE;
+	}
+
+	return bOk;
+}
+
+
+IMediaSamples::EFetchBestSampleResult FImgMediaLoader::FetchBestVideoSampleForTimeRange(const TRange<FMediaTimeStamp>& InTimeRange, TSharedPtr<IMediaTextureSample, ESPMode::ThreadSafe>& OutSample, bool bIsLoopingEnabled, float PlayRate)
+{
+	if (IsInitialized() && InTimeRange.HasLowerBound() && InTimeRange.HasUpperBound())
+	{
+		// A few sanity checks
+		check(InTimeRange.GetLowerBoundValue() <= InTimeRange.GetUpperBoundValue());
+		check(FMediaTimeStamp::GetPrimaryIndex(InTimeRange.GetLowerBoundValue().SequenceIndex) == FMediaTimeStamp::GetPrimaryIndex(InTimeRange.GetUpperBoundValue().SequenceIndex));
+
+		auto TimeRange = InTimeRange;
+
+		// Clamp the current range requested against the last known timestamp to ensure we return the "next" sample at all times...
+		if (QueuedSampleFetch.LastTimeStamp.IsValid())
+		{
+			if (PlayRate >= 0.0f)
+			{
+//WHAT IF THE LAST FRAME WAS SHOWN IN PAUSE MODE?
+				FMediaTimeStamp NextTimeStamp = QueuedSampleFetch.LastTimeStamp + QueuedSampleFetch.LastDuration;
+				if (NextTimeStamp >= TimeRange.GetUpperBoundValue())
+					{
+					return IMediaSamples::EFetchBestSampleResult::NoSample;
+					}
+
+				if (TimeRange.GetLowerBoundValue() < NextTimeStamp)
+				{
+					TimeRange = TRange<FMediaTimeStamp>(NextTimeStamp, InTimeRange.GetUpperBoundValue());
+				}
+			}
+			else
+			{
+				FMediaTimeStamp NextTimeStamp = QueuedSampleFetch.LastTimeStamp;
+				if (NextTimeStamp <= TimeRange.GetLowerBoundValue())
+				{
+					return IMediaSamples::EFetchBestSampleResult::NoSample;
+				}
+
+				if (TimeRange.GetUpperBoundValue() > NextTimeStamp)
+				{
+					TimeRange = TRange<FMediaTimeStamp>(InTimeRange.GetLowerBoundValue(), NextTimeStamp);
+				}
+			}
+		}
+
 		FTimespan StartTime = TimeRange.GetLowerBoundValue().Time;
-		int64 StartSequenceIndex = TimeRange.GetLowerBoundValue().SequenceIndex;
+		int32 StartLoopIndex = FMediaTimeStamp::GetSecondaryIndex(TimeRange.GetLowerBoundValue().SequenceIndex);
 		FTimespan EndTime = TimeRange.GetUpperBoundValue().Time;
-		int64 EndSequenceIndex = TimeRange.GetUpperBoundValue().SequenceIndex;
-		check(TimeRange.GetLowerBoundValue() <= TimeRange.GetUpperBoundValue());
+		int32 EndLoopIndex = FMediaTimeStamp::GetSecondaryIndex(TimeRange.GetUpperBoundValue().SequenceIndex);
 
 		if (bIsLoopingEnabled)
 		{
@@ -449,6 +531,8 @@ IMediaSamples::EFetchBestSampleResult FImgMediaLoader::FetchBestVideoSampleForTi
 		int32 MaxIdx = -1;
 		if (PlayRate >= 0.0f)
 		{
+			// Forward...
+
 			// Is the start index the same as the end index?
 			if (StartIndex == EndIndex)
 			{
@@ -456,7 +540,7 @@ IMediaSamples::EFetchBestSampleResult FImgMediaLoader::FetchBestVideoSampleForTi
 				FTimespan CurrentStartTime = FrameNumberToTime(StartIndex);
 				if (StartTime < CurrentStartTime)
 				{
-					// Yes. TimeToFrameNumber rounded up, but we really dont want that
+					// Yes. TimeToFrameNumber rounded up, but we really don't want that
 					// so adjust StartIndex.
 					if (StartIndex > 0)
 					{
@@ -465,9 +549,10 @@ IMediaSamples::EFetchBestSampleResult FImgMediaLoader::FetchBestVideoSampleForTi
 				}
 			}
 
-			// Forward...
-			if ((StartIndex > EndIndex) || (EndSequenceIndex > StartSequenceIndex))
+			// Does the range span across the end of the media?
+			if ((StartIndex > EndIndex) || (EndLoopIndex > StartLoopIndex))
 			{
+				// Yes...
 				int32 MaxIdx1, MaxIdx2;
 				float MaxOverlap1 = FindMaxOverlapInRange(StartIndex, GetNumImages()-1, StartTime, FrameNumberToTime(GetNumImages()), MaxIdx1);
 				float MaxOverlap2 = FindMaxOverlapInRange(0, EndIndex, FTimespan::Zero(), EndTime, MaxIdx2);
@@ -481,8 +566,27 @@ IMediaSamples::EFetchBestSampleResult FImgMediaLoader::FetchBestVideoSampleForTi
 		else
 		{
 			// Backward...
-			if ((StartIndex > EndIndex) || (StartSequenceIndex < EndSequenceIndex))
+			
+			// Is the start index the same as the end index?
+			if (StartIndex == EndIndex)
+			{
+				// Is StartTime larger than the time of StartIndex?
+				FTimespan CurrentEndTime = FrameNumberToTime(EndIndex);
+				if (EndTime > CurrentEndTime)
+				{
+					// Yes. TimeToFrameNumber rounded up, but we really don't want that
+					// so adjust StartIndex.
+					if (EndIndex < (GetNumImages() - 1))
+					{
+						++EndIndex;
+					}
+				}
+			}
+
+			// Does the range span across the end of the media?
+			if ((StartIndex > EndIndex) || (StartLoopIndex < EndLoopIndex))
 			{
+				// Yes...
 				int32 MaxIdx1, MaxIdx2;
 				float MaxOverlap1 = FindMaxOverlapInRange(EndIndex, 0, FTimespan::Zero(), EndTime, MaxIdx1);
 				float MaxOverlap2 = FindMaxOverlapInRange(GetNumImages()-1, StartIndex,  StartTime, FrameNumberToTime(GetNumImages()), MaxIdx2);
@@ -504,6 +608,7 @@ IMediaSamples::EFetchBestSampleResult FImgMediaLoader::FetchBestVideoSampleForTi
 
 			FScopeLock Lock(&CriticalSection);
 
+#if ALLOW_CHEATED_FRAMES
 			// If playback is not blocking, we expect less expectancy of precision on the users side, but more need for speedy return of "some ok frame"
 			// So: if we detect non-blocking playback we return a "as good sample as we can", but not always the "perfect" one we calculated
 			// (still we adhere to a rough emulation of a classic output pipeline as other players have)
@@ -563,6 +668,7 @@ IMediaSamples::EFetchBestSampleResult FImgMediaLoader::FetchBestVideoSampleForTi
 				}
 			}
 			else
+#endif
 			{
 				// Get a frame if we have one available right now...
 				Frame = UseGlobalCache ? GlobalCache->FindAndTouch(SequenceName, MaxIdx) : Frames.FindAndTouch(MaxIdx);
@@ -571,40 +677,44 @@ IMediaSamples::EFetchBestSampleResult FImgMediaLoader::FetchBestVideoSampleForTi
 			// Got a potential frame?
 			if (Frame)
 			{
+				// Yes.
+
+				/*
+					The player facade will pass in a time range with current primary and secondary sequence indices.
+					We need to make sure that we detect if we return a value "before" the last one (loop case) and
+					adjust the secondary index accordingly.
+				*/
+
 				RetryCount = 0;
-				double Duration = Frame->Get()->Info.FrameRate.AsInterval();
 
-				// Yes.
-				int32 NewSequenceIndex = 0;
-				
-				// Make sure this sample has the sequence index that matches the time range.
-				// If the time range only has one sequence index then just use that.
-				if (TimeRange.GetLowerBoundValue().SequenceIndex == TimeRange.GetUpperBoundValue().SequenceIndex)
-				{
-					NewSequenceIndex = TimeRange.GetLowerBoundValue().SequenceIndex;
-				}
-				else
+
+				// Check if we wrapped around the end of the media vs. the last frame we returned...
+				if (QueuedSampleFetch.LastFrameIndex != INDEX_NONE && bIsLoopingEnabled)
 				{
-					// Try the lower bound sequence index.
-					NewSequenceIndex = TimeRange.GetLowerBoundValue().SequenceIndex;
-					FMediaTimeStamp SampleTime = FMediaTimeStamp(FrameNumberToTime(MaxIdx), NewSequenceIndex);
-					TRange<FMediaTimeStamp> SampleTimeRange(
-						SampleTime,
-						SampleTime + FTimespan::FromSeconds(Duration));
-						
-					// Does our sample time overlap the time range?
-					if (TimeRange.Overlaps(SampleTimeRange) == false)
+					if (PlayRate >= 0.0f)
 					{
-						// No. Use the upper bound sequence index.
-						NewSequenceIndex = TimeRange.GetUpperBoundValue().SequenceIndex;
+						if (QueuedSampleFetch.LastFrameIndex > MaxIdx)
+						{
+							++QueuedSampleFetch.LoopIndex;
+						}
+					}
+					else
+					{
+						if (QueuedSampleFetch.LastFrameIndex < MaxIdx)
+						{
+							--QueuedSampleFetch.LoopIndex;
+						}
 					}
 				}
+
+				int64 NewSequenceIndex = FMediaTimeStamp::MakeSequenceIndex(FMediaTimeStamp::GetPrimaryIndex(TimeRange.GetLowerBoundValue().SequenceIndex), QueuedSampleFetch.LoopIndex);
 				
 				// Different from the last one we returned?
-				if ((QueuedSampleFetch.LastFrameIndex != MaxIdx) || (QueuedSampleFetch.CurrentSequenceIndex != NewSequenceIndex) || (ImagePaths.Num() == 1))
+				if ((QueuedSampleFetch.LastFrameIndex != MaxIdx) || (QueuedSampleFetch.LastTimeStamp.IsValid() && (QueuedSampleFetch.LastTimeStamp.SequenceIndex != NewSequenceIndex)) || (ImagePaths.Num() == 1))
 				{
 					QueuedSampleFetch.LastFrameIndex = MaxIdx;
-					QueuedSampleFetch.CurrentSequenceIndex = NewSequenceIndex;
+					QueuedSampleFetch.LastTimeStamp = FMediaTimeStamp(FrameNumberToTime(MaxIdx), NewSequenceIndex);
+					QueuedSampleFetch.LastDuration = FTimespan::FromSeconds(Frame->Get()->Info.FrameRate.AsInterval());
 
 					// We are clear to return it as new result... Make a sample & initialize it...
 					auto Sample = MakeShared<FImgMediaTextureSample, ESPMode::ThreadSafe>();
@@ -612,7 +722,7 @@ IMediaSamples::EFetchBestSampleResult FImgMediaLoader::FetchBestVideoSampleForTi
 
 					ImgMediaLoader::CheckAndUpdateImgDimensions(SequenceDim, Frame->Get()->Info.Dim);
 
-					if (Sample->Initialize(*Frame->Get(), SequenceDim, FMediaTimeStamp(FrameNumberToTime(MaxIdx), QueuedSampleFetch.CurrentSequenceIndex), FTimespan::FromSeconds(Duration), GetNumMipLevels(), TilingDescription))
+					if (Sample->Initialize(*Frame->Get(), SequenceDim, QueuedSampleFetch.LastTimeStamp, QueuedSampleFetch.LastDuration, GetNumMipLevels(), TilingDescription))
 					{
 						OutSample = Sample;
 						CSV_EVENT(ImgMedia, TEXT("LoaderFetchHit %d %d-%d"), MaxIdx, StartIndex, EndIndex);
@@ -660,7 +770,7 @@ bool FImgMediaLoader::PeekVideoSampleTime(FMediaTimeStamp &TimeStamp, bool bIsLo
 	if (IsInitialized())
 	{
 		int32 Idx;
-		bool bNewSeq = false;
+		bool bNewLoopSeq = false;
 
 		// Do we know which index we handed out last?
 		if (QueuedSampleFetch.LastFrameIndex != INDEX_NONE)
@@ -673,12 +783,12 @@ bool FImgMediaLoader::PeekVideoSampleTime(FMediaTimeStamp &TimeStamp, bool bIsLo
 				if (Idx < 0)
 				{
 					Idx = NumFrames - 1;
-					bNewSeq = true;
+					bNewLoopSeq = true;
 				}
 				else if (Idx >= NumFrames)
 				{
 					Idx = 0;
-					bNewSeq = true;
+					bNewLoopSeq = true;
 				}
 			}
 			else
@@ -692,10 +802,45 @@ bool FImgMediaLoader::PeekVideoSampleTime(FMediaTimeStamp &TimeStamp, bool bIsLo
 		}
 		else
 		{
-			// No, we don't have an index. Just compute things based on the current time given...
-			Idx = TimeToFrameNumber(CurrentTime);
+			if (QueuedSampleFetch.LastTimeStamp.IsValid())
+			{
+				if (PlayRate >= 0.0f)
+				{
+					auto NextTime = QueuedSampleFetch.LastTimeStamp.Time + QueuedSampleFetch.LastDuration;
+					if (NextTime >= SequenceDuration)
+					{
+						if (!bIsLoopingEnabled)
+						{
+							return false;
+						}
+						NextTime -= SequenceDuration;
+						bNewLoopSeq = true;
+					}
+					Idx = TimeToFrameNumber(NextTime);
+				}
+				else
+				{
+					auto NextTime = QueuedSampleFetch.LastTimeStamp.Time - QueuedSampleFetch.LastDuration;
+					if (NextTime < FTimespan::Zero())
+					{
+						if (!bIsLoopingEnabled)
+						{
+							return false;
+						}
+						NextTime += SequenceDuration;
+						bNewLoopSeq = true;
+					}
+					Idx = TimeToFrameNumber(NextTime);
+				}
+			}
+			else
+			{
+				Idx = TimeToFrameNumber(CurrentTime);
+			}
 		}
 
+		auto SequenceIndex = QueuedSampleFetch.LastTimeStamp.IsValid() ? QueuedSampleFetch.LastTimeStamp.SequenceIndex : 0;
+
 		if (SkipFramesLevel > 0)
 		{
 			Idx = GetSkipFrame(Idx);
@@ -713,7 +858,7 @@ bool FImgMediaLoader::PeekVideoSampleTime(FMediaTimeStamp &TimeStamp, bool bIsLo
 		{
 			// Yes. Return the timing information...
 			TimeStamp.Time = FrameStart;
-			TimeStamp.SequenceIndex = bNewSeq ? (QueuedSampleFetch.CurrentSequenceIndex + 1) : QueuedSampleFetch.CurrentSequenceIndex;
+			TimeStamp.SequenceIndex = bNewLoopSeq ? FMediaTimeStamp::AdjustSecondaryIndex(SequenceIndex, (PlayRate >= 0.0f) ? 1 : -1) : SequenceIndex;
 			return true;
 		}
 		else
@@ -726,6 +871,18 @@ bool FImgMediaLoader::PeekVideoSampleTime(FMediaTimeStamp &TimeStamp, bool bIsLo
 }
 
 
+bool FImgMediaLoader::IsFrameFirst(const FTimespan& TimeStamp) const
+{
+	return TimeToFrameNumber(TimeStamp) == 0;
+}
+
+
+bool FImgMediaLoader::IsFrameLast(const FTimespan & TimeStamp) const
+{
+	return TimeToFrameNumber(TimeStamp) == (GetNumImages() - 1);
+}
+
+
 void FImgMediaLoader::GetPendingTimeRanges(TRangeSet<FTimespan>& OutRangeSet) const
 {
 	FScopeLock Lock(&CriticalSection);
diff --git a/Engine/Plugins/Media/ImgMedia/Source/ImgMedia/Private/Player/ImgMediaPlayer.cpp b/Engine/Plugins/Media/ImgMedia/Source/ImgMedia/Private/Player/ImgMediaPlayer.cpp
--- a/Engine/Plugins/Media/ImgMedia/Source/ImgMedia/Private/Player/ImgMediaPlayer.cpp
+++ b/Engine/Plugins/Media/ImgMedia/Source/ImgMedia/Private/Player/ImgMediaPlayer.cpp
@@ -76,8 +76,10 @@ FImgMediaPlayer::FImgMediaPlayer(IMediaEventSink& InEventSink, const TSharedRef<
 	const TSharedRef<FImgMediaGlobalCache, ESPMode::ThreadSafe>& InGlobalCache)
 	: CurrentDuration(FTimespan::Zero())
 	, CurrentRate(0.0f)
+	, LastNonZeroRate(0.0f)
 	, CurrentState(EMediaState::Closed)
 	, CurrentTime(FTimespan::Zero())
+	, CurrentSeekIndex(0)
 	, DeltaTimeHackApplied(false)
 	, EventSink(InEventSink)
 	, LastFetchTime(FTimespan::MinValue())
@@ -126,8 +128,10 @@ void FImgMediaPlayer::Close()
 	CurrentDuration = FTimespan::Zero();
 	CurrentUrl.Empty();
 	CurrentRate = 0.0f;
+	LastNonZeroRate = 0.0f;
 	CurrentState = EMediaState::Closed;
 	CurrentTime = FTimespan::Zero();
+	CurrentSeekIndex = 0;
 	DeltaTimeHackApplied = false;
 	LastFetchTime = FTimespan::MinValue();
 	PlaybackRestarted = false;
@@ -392,6 +396,7 @@ void FImgMediaPlayer::TickInput(FTimespan DeltaTime, FTimespan /*Timecode*/)
 			CurrentState = EMediaState::Stopped;
 			CurrentTime = FTimespan::Zero();
 			CurrentRate = 0.0f;
+			LastNonZeroRate = 0.0f;
 			DeltaTimeHackApplied = false;
 
 			EventSink.ReceiveMediaEvent(EMediaEvent::PlaybackSuspended);
@@ -422,8 +427,7 @@ bool FImgMediaPlayer::FlushOnSeekStarted() const
 #if IMG_MEDIA_PLAYER_VERSION == 1
 	return false;
 #else
-	// Flush on start, otherwise if we already have the frames ready during a seek,
-	// the frames we return will get flushed straight away.
+	// V2 player response will always be treated as 'true'
 	return true;
 #endif
 }
@@ -433,6 +437,7 @@ bool FImgMediaPlayer::FlushOnSeekCompleted() const
 #if IMG_MEDIA_PLAYER_VERSION == 1
 	return true;
 #else
+	// V2 player response will always be treated as 'false'
 	return false;
 #endif
 }
@@ -473,6 +478,7 @@ bool FImgMediaPlayer::GetPlayerFeatureFlag(EFeatureFlag flag) const
 	switch (flag)
 	{
 	case EFeatureFlag::UsePlaybackTimingV2:
+	case EFeatureFlag::PlayerUsesInternalFlushOnSeek:
 		return true;
 	default:
 		break;
@@ -645,6 +651,11 @@ bool FImgMediaPlayer::Seek(const FTimespan& Time)
 	}
 #else
 	CurrentTime = Time;
+	CurrentSeekIndex += (CurrentRate >= 0.0f) ? 1 : -1;
+	if (Loader.IsValid())
+	{
+		Loader->Seek(FMediaTimeStamp(CurrentTime, FMediaTimeStamp::MakeSequenceIndex(CurrentSeekIndex, 0)));
+	}
 #endif // IMG_MEDIA_PLAYER_VERSION == 1
 
 	if (CurrentState == EMediaState::Paused)
@@ -702,16 +713,18 @@ bool FImgMediaPlayer::SetRate(float Rate)
 		}
 
 		CurrentRate = Rate;
+		LastNonZeroRate = Rate;
 		CurrentState = EMediaState::Playing;
 
 		EventSink.ReceiveMediaEvent(EMediaEvent::PlaybackResumed);
 
 		return true;
 	}
-	
+
 	// handle pausing
 	if ((CurrentRate != 0.0f) && (Rate == 0.0f))
 	{
+		LastNonZeroRate = CurrentRate;
 		CurrentRate = Rate;
 		CurrentState = EMediaState::Paused;
 
@@ -724,6 +737,10 @@ bool FImgMediaPlayer::SetRate(float Rate)
 		return true;
 	}
 
+	if (Rate != 0.0f)
+	{
+		LastNonZeroRate = Rate;
+	}
 	CurrentRate = Rate;
 
 	return true;
@@ -785,49 +802,49 @@ void FImgMediaPlayer::FlushSamples()
 {
 	if (IsInitialized())
 	{
+		CurrentSeekIndex = 0;
 		LastFetchTime = FTimespan::MinValue();
 		Loader->ResetFetchLogic();
 	}
 }
 
 
+bool FImgMediaPlayer::DiscardVideoSamples(const TRange<FMediaTimeStamp>& TimeRange, bool bReverse)
+{
+	if (Loader.IsValid() && IsInitialized())
+	{
+		return Loader->DiscardVideoSamples(TimeRange, ShouldLoop, CurrentRate);
+	}
+	return false;
+}
+
+
 IMediaSamples::EFetchBestSampleResult FImgMediaPlayer::FetchBestVideoSampleForTimeRange(const TRange<FMediaTimeStamp>& TimeRange, TSharedPtr<IMediaTextureSample, ESPMode::ThreadSafe>& OutSample, bool bReverse)
 {
 	IMediaSamples::EFetchBestSampleResult SampleResult = EFetchBestSampleResult::NoSample;
 
-	// The facade will keep on asking for frames, so don't do anything if we are stopped.
-	if (Loader.IsValid() && IsInitialized() && (CurrentState != EMediaState::Stopped))
+	if (Loader.IsValid() && IsInitialized())
 	{
 		// See if we have any samples in the specified time range.
-		SampleResult = Loader->FetchBestVideoSampleForTimeRange(TimeRange, OutSample, ShouldLoop, CurrentRate);
+		SampleResult = Loader->FetchBestVideoSampleForTimeRange(TimeRange, OutSample, ShouldLoop, LastNonZeroRate);
 		Scheduler->TickInput(FTimespan::Zero(), FTimespan::MinValue());
 		if (SampleResult == IMediaSamples::EFetchBestSampleResult::Ok)
 		{
 			CurrentTime = OutSample->GetTime().Time;
-		}
 
-		// Are we not looping?
-		if (ShouldLoop == false)
-		{
-			// Are we at the end?
-			bool bIsAtEnd = false;
-			if (CurrentRate >= 0.0f)
+			// Are we not looping?
+			if (!ShouldLoop)
 			{
-				bIsAtEnd = ((TimeRange.HasUpperBound()) &&
-					(TimeRange.GetUpperBoundValue().Time >= CurrentDuration));
-			}
-			else
-			{
-				bIsAtEnd = ((TimeRange.HasLowerBound()) &&
-					(TimeRange.GetLowerBoundValue().Time <= 0.0f));
-			}
-			if (bIsAtEnd)
-			{
-				// Stop the player.
-				EventSink.ReceiveMediaEvent(EMediaEvent::PlaybackEndReached);
-				CurrentState = EMediaState::Stopped;
-				CurrentRate = 0.0f;
-				EventSink.ReceiveMediaEvent(EMediaEvent::PlaybackSuspended);
+				// Is this the last frame?
+				if ((CurrentRate < 0.0f) ? Loader->IsFrameFirst(CurrentTime) : Loader->IsFrameLast(CurrentTime))
+				{
+					// Yes. Stop the player...
+					EventSink.ReceiveMediaEvent(EMediaEvent::PlaybackEndReached);
+					CurrentState = EMediaState::Stopped;
+					CurrentRate = 0.0f;
+					LastNonZeroRate = 0.0f;
+					EventSink.ReceiveMediaEvent(EMediaEvent::PlaybackSuspended);
+				}
 			}
 		}
 	}
@@ -840,7 +857,7 @@ bool FImgMediaPlayer::PeekVideoSampleTime(FMediaTimeStamp& TimeStamp)
 	if (Loader.IsValid() && IsInitialized())
 	{
 		// Do we have the current frame?
-		return Loader->PeekVideoSampleTime(TimeStamp, ShouldLoop, CurrentRate, GetTime());
+		return Loader->PeekVideoSampleTime(TimeStamp, ShouldLoop, LastNonZeroRate, GetTime());
 	}
 	return false;
 }
diff --git a/Engine/Plugins/Media/MediaCompositing/Source/MediaCompositing/Private/MovieScene/MovieSceneMediaTemplate.cpp b/Engine/Plugins/Media/MediaCompositing/Source/MediaCompositing/Private/MovieScene/MovieSceneMediaTemplate.cpp
--- a/Engine/Plugins/Media/MediaCompositing/Source/MediaCompositing/Private/MovieScene/MovieSceneMediaTemplate.cpp
+++ b/Engine/Plugins/Media/MediaCompositing/Source/MediaCompositing/Private/MovieScene/MovieSceneMediaTemplate.cpp
@@ -222,9 +222,10 @@ struct FMediaSectionExecutionToken
 				//  the direction is taken into account as hint for internal operation of the player)
 				if (!MediaPlayer->SetRate((Context.GetDirection() == EPlayDirection::Forwards) ? 1.0f : -1.0f))
 				{
-					// Failed to set needed rate: better switch off blocking and bail...
-					MediaPlayer->SetBlockOnTimeRange(TRange<FTimespan>::Empty());
-					return;
+					// Failed to set needed rate. Keep things blocked, as this means the player will still not be playing, this will
+					// trigger a seek to each and every frame. A potentially very SLOW method of approximating backwards playback, but better
+					// than nothing.
+					// -> nothing to do
 				}
 			}
 			else
@@ -239,18 +240,20 @@ struct FMediaSectionExecutionToken
 				{
 					if (!MediaPlayer->SetRate(1.0f))
 					{
-						// Failed to set needed rate: better switch off blocking and bail...
-						MediaPlayer->SetBlockOnTimeRange(TRange<FTimespan>::Empty());
-						return;
+					// Failed to set needed rate. Keep things blocked, as this means the player will still be returning the old rate, we will get here repeatedly
+					// and each time trigger a seek. A potentially very SLOW method of approximating backwards playback, but better
+					// than nothing.
+					MediaPlayer->Seek(MediaTime);
 					}
 				}
 				else if (Context.GetDirection() == EPlayDirection::Backwards && CurrentPlayerRate > 0.0f)
 				{
 					if (!MediaPlayer->SetRate(-1.0f))
 					{
-						// Failed to set needed rate: better switch off blocking and bail...
-						MediaPlayer->SetBlockOnTimeRange(TRange<FTimespan>::Empty());
-						return;
+					// Failed to set needed rate. Keep things blocked, as this means the player will still be returning the old rate, we will get here repeatedly
+					// and each time trigger a seek. A potentially very SLOW method of approximating backwards playback, but better
+					// than nothing.
+					MediaPlayer->Seek(MediaTime);
 					}
 				}
 			}
@@ -265,8 +268,8 @@ struct FMediaSectionExecutionToken
 			MediaPlayer->Seek(MediaTime);
 		}
 
-	// Set blocking range / time-range to display
-	// (we always use the full current time for this, any adjustments to player timestamps are done internally)
+		// Set blocking range / time-range to display
+		// (we always use the full current time for this, any adjustments to player timestamps are done internally)
 		MediaPlayer->SetBlockOnTimeRange(TRange<FTimespan>(CurrentTime, CurrentTime + FrameDuration));
 	}
 
diff --git a/Engine/Plugins/Media/MediaPlayerEditor/Source/MediaPlayerEditor/Private/Widgets/SMediaPlayerEditorCache.cpp b/Engine/Plugins/Media/MediaPlayerEditor/Source/MediaPlayerEditor/Private/Widgets/SMediaPlayerEditorCache.cpp
--- a/Engine/Plugins/Media/MediaPlayerEditor/Source/MediaPlayerEditor/Private/Widgets/SMediaPlayerEditorCache.cpp
+++ b/Engine/Plugins/Media/MediaPlayerEditor/Source/MediaPlayerEditor/Private/Widgets/SMediaPlayerEditorCache.cpp
@@ -63,7 +63,7 @@ int32 SMediaPlayerEditorCache::OnPaint(const FPaintArgs& Args, const FGeometry&
 	{
 		DrawSampleCache(EMediaTrackType::Audio, AllottedGeometry, OutDrawElements, LayerId, InWidgetStyle, 0.0f, 0.5f);
 		DrawSampleCache(EMediaTrackType::Video, AllottedGeometry, OutDrawElements, LayerId, InWidgetStyle, 0.5f, 0.5f);
-		DrawPlayerPosition(MediaPlayer->GetTime(), AllottedGeometry, OutDrawElements, LayerId, InWidgetStyle, FLinearColor::Gray);
+		DrawPlayerPosition(MediaPlayer->GetDisplayTime(), AllottedGeometry, OutDrawElements, LayerId, InWidgetStyle, FLinearColor::Gray);
 	}
 
 	return LayerId + 1;
diff --git a/Engine/Plugins/Media/MediaPlayerEditor/Source/MediaPlayerEditor/Private/Widgets/SMediaPlayerEditorViewer.cpp b/Engine/Plugins/Media/MediaPlayerEditor/Source/MediaPlayerEditor/Private/Widgets/SMediaPlayerEditorViewer.cpp
--- a/Engine/Plugins/Media/MediaPlayerEditor/Source/MediaPlayerEditor/Private/Widgets/SMediaPlayerEditorViewer.cpp
+++ b/Engine/Plugins/Media/MediaPlayerEditor/Source/MediaPlayerEditor/Private/Widgets/SMediaPlayerEditorViewer.cpp
@@ -381,7 +381,7 @@ void SMediaPlayerEditorViewer::Construct(const FArguments& InArgs, UMediaPlayer&
 													})
 													.OnMouseCaptureBegin_Lambda([this]()
 													{
-														ScrubValue = FTimespan::Ratio(MediaPlayer->GetTime(), MediaPlayer->GetDuration());
+														ScrubValue = FTimespan::Ratio(MediaPlayer->GetDisplayTime(), MediaPlayer->GetDuration());
 
 														if (MediaPlayer->SupportsScrubbing())
 														{
@@ -416,7 +416,7 @@ void SMediaPlayerEditorViewer::Construct(const FArguments& InArgs, UMediaPlayer&
 														}
 
 														return (MediaPlayer->GetDuration() > FTimespan::Zero())
-															? FTimespan::Ratio(MediaPlayer->GetTime(), MediaPlayer->GetDuration())
+															? FTimespan::Ratio(MediaPlayer->GetDisplayTime(), MediaPlayer->GetDuration())
 															: 0.0f;
 													})
 													.Visibility_Lambda([this]() -> EVisibility {
@@ -943,7 +943,7 @@ FText SMediaPlayerEditorViewer::HandleTimerTextBlockText() const
 		return FText::GetEmpty();
 	}
 
-	FTimespan Time = MediaPlayer->GetTime();
+	FTimespan Time = MediaPlayer->GetDisplayTime();
 
 	// empty string if time n/a
 	if (Time < FTimespan::Zero())
@@ -976,7 +976,7 @@ FText SMediaPlayerEditorViewer::HandleTimerTextBlockToolTipText() const
 	}
 
 	FTimespan Duration = MediaPlayer->GetDuration();
-	FTimespan Remaining = Duration - MediaPlayer->GetTime();
+	FTimespan Remaining = Duration - MediaPlayer->GetDisplayTime();
 	bool bInfiniteTimeRemaining = Duration == FTimespan::MaxValue();
 
 	if (Remaining <= FTimespan::Zero() && !bInfiniteTimeRemaining)
diff --git a/Engine/Plugins/Media/WmfMedia/Source/WmfMedia/Private/Player/WmfMediaTracks.cpp b/Engine/Plugins/Media/WmfMedia/Source/WmfMedia/Private/Player/WmfMediaTracks.cpp
--- a/Engine/Plugins/Media/WmfMedia/Source/WmfMedia/Private/Player/WmfMediaTracks.cpp
+++ b/Engine/Plugins/Media/WmfMedia/Source/WmfMedia/Private/Player/WmfMediaTracks.cpp
@@ -371,6 +371,7 @@ void FWmfMediaTracks::Shutdown()
 	SelectionChanged = false;
 #if WMFMEDIA_PLAYER_VERSION >= 2
 	SeekTimeOptional.Reset();
+	SeekIndex = 0;
 #endif // WMFMEDIA_PLAYER_VERSION >= 2
 }
 
@@ -386,6 +387,7 @@ void FWmfMediaTracks::SeekStarted(const FTimespan& InTime)
 {
 	UE_LOG(LogWmfMedia, VeryVerbose, TEXT("FWmfMediaTracks::SeekStarted %f"), InTime.GetTotalSeconds());
 	SeekTimeOptional = InTime;
+	++SeekIndex;
 }
 
 #endif // WMFMEDIA_PLAYER_VERSION >= 2
@@ -552,11 +554,7 @@ IMediaSamples::EFetchBestSampleResult FWmfMediaTracks::FetchBestVideoSampleForTi
 			if (SeekTimeOptional.IsSet())
 			{
 				// Are we past the seek time?
-				if (TimeRangeTime.Contains(SeekTimeOptional.GetValue()) == false)
-				{
-					SeekTimeOptional.Reset();
-				}
-				else
+				if (TimeRangeTime.Contains(SeekTimeOptional.GetValue()))
 				{
 					// Is this our seek sample?
 					FTimespan SeekTime = SeekTimeOptional.GetValue();
@@ -634,9 +632,8 @@ IMediaSamples::EFetchBestSampleResult FWmfMediaTracks::FetchBestVideoSampleForTi
 							// Update sequence index.
 							FWmfMediaTextureSample* WmfSample =
 								static_cast<FWmfMediaTextureSample*>(OutSample.Get());
-							WmfSample->SetSequenceIndex(GetSequenceIndex(TimeRange,
-								WmfSample->GetTime().Time));
-								
+							WmfSample->SetSequenceIndex(FMediaTimeStamp::MakeSequenceIndex(SeekIndex, 0));
+
 							UE_LOG(LogWmfMedia, VeryVerbose, TEXT("FetchBestVideoSampleForTimeRange got sample."));
 						}
 					}
@@ -2000,33 +1997,6 @@ const FWmfMediaTracks::FFormat* FWmfMediaTracks::GetVideoFormat(int32 TrackIndex
 }
 
 
-int64 FWmfMediaTracks::GetSequenceIndex(const TRange<FMediaTimeStamp>& TimeRange, FTimespan Time) const
-{
-	int64 SequenceIndex = 0;
-
-	// Make sure this has the sequence index that matches the time range.
-	// If the time range only has one sequence index then just use that.
-	if (TimeRange.GetLowerBoundValue().SequenceIndex == TimeRange.GetUpperBoundValue().SequenceIndex)
-	{
-		SequenceIndex = TimeRange.GetLowerBoundValue().SequenceIndex;
-	}
-	else
-	{
-		// Try the lower bound sequence index.
-		SequenceIndex = TimeRange.GetLowerBoundValue().SequenceIndex;
-		FMediaTimeStamp SampleTime = FMediaTimeStamp(Time, SequenceIndex);
-
-		// Does our time overlap the time range?
-		if (TimeRange.Contains(SampleTime) == false)
-		{
-			// No. Use the upper bound sequence index.
-			SequenceIndex = TimeRange.GetUpperBoundValue().SequenceIndex;
-		}
-	}
-
-	return SequenceIndex;
-}
-
 /* FWmfMediaTracks callbacks
  *****************************************************************************/
 
diff --git a/Engine/Plugins/Media/WmfMedia/Source/WmfMedia/Private/Player/WmfMediaTracks.h b/Engine/Plugins/Media/WmfMedia/Source/WmfMedia/Private/Player/WmfMediaTracks.h
--- a/Engine/Plugins/Media/WmfMedia/Source/WmfMedia/Private/Player/WmfMediaTracks.h
+++ b/Engine/Plugins/Media/WmfMedia/Source/WmfMedia/Private/Player/WmfMediaTracks.h
@@ -102,7 +102,7 @@ public:
 	 *
 	 * @param OutStats The string to append the statistics to.
 	 */
-	void AppendStats(FString &OutStats) const;
+	void AppendStats(FString& OutStats) const;
 
 	/**
 	 * Clear the streams flags.
@@ -205,6 +205,12 @@ public:
 	virtual bool FetchCaption(TRange<FTimespan> TimeRange, TSharedPtr<IMediaOverlaySample, ESPMode::ThreadSafe>& OutSample) override;
 	virtual bool FetchMetadata(TRange<FTimespan> TimeRange, TSharedPtr<IMediaBinarySample, ESPMode::ThreadSafe>& OutSample) override;
 	virtual bool FetchVideo(TRange<FTimespan> TimeRange, TSharedPtr<IMediaTextureSample, ESPMode::ThreadSafe>& OutSample) override;
+
+	virtual bool FetchAudio(TRange<FMediaTimeStamp> TimeRange, TSharedPtr<IMediaAudioSample, ESPMode::ThreadSafe>& OutSample) override { return FetchAudio(TimeRange.GetLowerBound().IsOpen() ? TRange<FTimespan>() : TRange<FTimespan>(TimeRange.GetLowerBoundValue().Time, TimeRange.GetUpperBoundValue().Time), OutSample); }
+	virtual bool FetchCaption(TRange<FMediaTimeStamp> TimeRange, TSharedPtr<IMediaOverlaySample, ESPMode::ThreadSafe>& OutSample) override { return FetchCaption(TRange<FTimespan>(TimeRange.GetLowerBoundValue().Time, TimeRange.GetUpperBoundValue().Time), OutSample); }
+	virtual bool FetchMetadata(TRange<FMediaTimeStamp> TimeRange, TSharedPtr<IMediaBinarySample, ESPMode::ThreadSafe>& OutSample) override { return FetchMetadata(TRange<FTimespan>(TimeRange.GetLowerBoundValue().Time, TimeRange.GetUpperBoundValue().Time), OutSample); }
+	virtual bool FetchVideo(TRange<FMediaTimeStamp> TimeRange, TSharedPtr<IMediaTextureSample, ESPMode::ThreadSafe>& OutSample) override { return FetchVideo(TRange<FTimespan>(TimeRange.GetLowerBoundValue().Time, TimeRange.GetUpperBoundValue().Time), OutSample); }
+
 	virtual void FlushSamples() override;
 #if WMFMEDIA_PLAYER_VERSION >= 2
 	virtual EFetchBestSampleResult FetchBestVideoSampleForTimeRange(const TRange<FMediaTimeStamp> & TimeRange, TSharedPtr<IMediaTextureSample, ESPMode::ThreadSafe>& OutSample, bool bReverse) override;
@@ -282,15 +288,6 @@ private:
 	 */
 	const FFormat* GetVideoFormat(int32 TrackIndex, int32 FormatIndex) const;
 
-	/**
-	 * Gets the sequence index for a specific time that fits the time range.
-	 * 
-	 * @param	TimeRange		Sequence index will be in this range.
-	 * @param	Time			Time to retrieve sequence index for.
-	 * @return					Sequence index.
-	 */
-	int64 GetSequenceIndex(const TRange<FMediaTimeStamp>& TimeRange, FTimespan Time) const;
-
 private:
 
 	/** Callback for handling media sampler pauses. */
@@ -386,6 +383,8 @@ private:
 #if WMFMEDIA_PLAYER_VERSION >= 2
 	/** If set, then discard samples until we get this sample. */
 	TOptional<FTimespan> SeekTimeOptional;
+	/** Seek index tracking */
+	int32 SeekIndex;
 #endif // WMFMEDIA_PLAYER_VERSION >= 2
 };
 
diff --git a/Engine/Plugins/Media/WmfMedia/Source/WmfMedia/Private/Wmf/WmfMediaSession.cpp b/Engine/Plugins/Media/WmfMedia/Source/WmfMedia/Private/Wmf/WmfMediaSession.cpp
--- a/Engine/Plugins/Media/WmfMedia/Source/WmfMedia/Private/Wmf/WmfMediaSession.cpp
+++ b/Engine/Plugins/Media/WmfMedia/Source/WmfMedia/Private/Wmf/WmfMediaSession.cpp
@@ -500,7 +500,7 @@ bool FWmfMediaSession::Seek(const FTimespan& Time)
 		return true;
 	}
 
-	return CommitTime(Time);
+	return CommitTime(Time, true);
 }
 
 
@@ -961,7 +961,7 @@ bool FWmfMediaSession::CommitRate(float Rate)
 		else
 		{
 			UE_LOG(LogWmfMedia, Verbose, TEXT("Session %p: Starting session for rate change"), this);
-			CommitTime(RestartTime);
+			CommitTime(RestartTime, false);
 		}
 	}
 
@@ -969,7 +969,7 @@ bool FWmfMediaSession::CommitRate(float Rate)
 }
 
 
-bool FWmfMediaSession::CommitTime(FTimespan Time)
+bool FWmfMediaSession::CommitTime(FTimespan Time, bool bIsSeek)
 {
 	check(MediaSession != NULL);
 	check(PendingChanges == false);
@@ -1020,7 +1020,7 @@ bool FWmfMediaSession::CommitTime(FTimespan Time)
 
 #if WMFMEDIA_PLAYER_VERSION >= 2
 	// If this is not a loop, then tell the tracks about the seek.
-	if ((bIsRequestedTimeLoop == false) && (bCanSeek))
+	if ((bIsRequestedTimeLoop == false) && (bCanSeek) && bIsSeek)
 	{
 		TSharedPtr<FWmfMediaTracks, ESPMode::ThreadSafe> TracksPinned = Tracks.Pin();
 		if (TracksPinned.IsValid())
@@ -1171,7 +1171,7 @@ void FWmfMediaSession::DoPendingChanges()
 		const FTimespan Time = RequestedTime.GetValue();
 		RequestedTime.Reset();
 
-		CommitTime(Time);
+		CommitTime(Time, false);
 		bIsRequestedTimeLoop = false;
 	}
 }
diff --git a/Engine/Plugins/Media/WmfMedia/Source/WmfMedia/Private/Wmf/WmfMediaSession.h b/Engine/Plugins/Media/WmfMedia/Source/WmfMedia/Private/Wmf/WmfMediaSession.h
--- a/Engine/Plugins/Media/WmfMedia/Source/WmfMedia/Private/Wmf/WmfMediaSession.h
+++ b/Engine/Plugins/Media/WmfMedia/Source/WmfMedia/Private/Wmf/WmfMediaSession.h
@@ -143,7 +143,7 @@ protected:
 	 * @param Time The play position to commit.
 	 * @see CommitRate, CommitTopology
 	 */
-	bool CommitTime(FTimespan Time);
+	bool CommitTime(FTimespan Time, bool bIsSeek);
 
 	/**
 	 * Commit the given playback topology.
diff --git a/Engine/Source/Runtime/Media/Public/IMediaSamples.h b/Engine/Source/Runtime/Media/Public/IMediaSamples.h
--- a/Engine/Source/Runtime/Media/Public/IMediaSamples.h
+++ b/Engine/Source/Runtime/Media/Public/IMediaSamples.h
@@ -71,6 +71,10 @@ public:
 	{
 		return false; // override in child classes, if supported
 	}
+	virtual bool FetchMetadata(TRange<FMediaTimeStamp> TimeRange, TSharedPtr<IMediaBinarySample, ESPMode::ThreadSafe>& OutSample)
+	{
+		return false; // override in child classes, if supported
+	}
 
 	/**
 	 * Fetch the next subtitle sample.
@@ -125,8 +129,16 @@ public:
 
 	virtual bool PeekVideoSampleTime(FMediaTimeStamp & TimeStamp) = 0;
 
-	virtual uint32 PurgeOutdatedVideoSamples(const FMediaTimeStamp & ReferenceTime, bool bReversed) { return 0; };
-	virtual uint32 PurgeOutdatedSubtitleSamples(const FMediaTimeStamp & ReferenceTime, bool bReversed) { return 0; };
+	virtual bool DiscardVideoSamples(const TRange<FMediaTimeStamp>& TimeRange, bool bReverse) { return false; }
+	virtual bool DiscardAudioSamples(const TRange<FMediaTimeStamp>& TimeRange, bool bReverse) { return false; }
+	virtual bool DiscardCaptionSamples(const TRange<FMediaTimeStamp>& TimeRange, bool bReverse) { return false; }
+	virtual bool DiscardSubtitleSamples(const TRange<FMediaTimeStamp>& TimeRange, bool bReverse) { return false; }
+	virtual bool DiscardMetadataSamples(const TRange<FMediaTimeStamp>& TimeRange, bool bReverse) { return false; }
+
+	virtual uint32 PurgeOutdatedVideoSamples(const FMediaTimeStamp & ReferenceTime, bool bReversed, FTimespan MaxAge) { return 0; };
+	virtual uint32 PurgeOutdatedCaptionSamples(const FMediaTimeStamp& ReferenceTime, bool bReversed, FTimespan MaxAge) { return 0; };
+	virtual uint32 PurgeOutdatedSubtitleSamples(const FMediaTimeStamp & ReferenceTime, bool bReversed, FTimespan MaxAge) { return 0; };
+	virtual uint32 PurgeOutdatedMetadataSamples(const FMediaTimeStamp& ReferenceTime, bool bReversed, FTimespan MaxAge) { return 0; };
 
 	virtual bool CanReceiveVideoSamples(uint32 Num) const { return true; }
 	virtual bool CanReceiveAudioSamples(uint32 Num) const { return true; }
diff --git a/Engine/Source/Runtime/Media/Public/IMediaTimeSource.h b/Engine/Source/Runtime/Media/Public/IMediaTimeSource.h
--- a/Engine/Source/Runtime/Media/Public/IMediaTimeSource.h
+++ b/Engine/Source/Runtime/Media/Public/IMediaTimeSource.h
@@ -42,6 +42,32 @@ public:
 	bool operator > (const FMediaTimeStamp & Other) const { return (SequenceIndex > Other.SequenceIndex) || (SequenceIndex == Other.SequenceIndex && Time > Other.Time); }
 	bool operator >= (const FMediaTimeStamp & Other) const { return (SequenceIndex > Other.SequenceIndex) || (SequenceIndex == Other.SequenceIndex && Time >= Other.Time); }
 
+	static int64 MakeSequenceIndex(int32 PrimaryIndex, int32 SecondaryIndex)
+	{
+		return (int64(PrimaryIndex) << 32) + int64(SecondaryIndex);
+	}
+
+	static int64 AdjustPrimaryIndex(int64 SequenceIndex, int32 Add)
+	{
+		return SequenceIndex + (int64(Add) << 32);
+	}
+
+	static int64 AdjustSecondaryIndex(int64 SequenceIndex, int32 Add)
+	{
+		return SequenceIndex + Add;
+	}
+
+	static int32 GetPrimaryIndex(int64 SequenceIndex)
+	{
+		// note: needs to cope with signed operation
+		return int32((SequenceIndex - int32(SequenceIndex)) >> 32);
+	}
+
+	static int32 GetSecondaryIndex(int64 SequenceIndex)
+	{
+		return int32(SequenceIndex);
+	}
+
 	FMediaTimeStamp operator + (const FTimespan & Other) const { return FMediaTimeStamp(Time + Other, SequenceIndex); }
 	FMediaTimeStamp operator - (const FTimespan & Other) const { return FMediaTimeStamp(Time - Other, SequenceIndex); }
 
diff --git a/Engine/Source/Runtime/MediaAssets/Private/Assets/MediaPlayer.cpp b/Engine/Source/Runtime/MediaAssets/Private/Assets/MediaPlayer.cpp
--- a/Engine/Source/Runtime/MediaAssets/Private/Assets/MediaPlayer.cpp
+++ b/Engine/Source/Runtime/MediaAssets/Private/Assets/MediaPlayer.cpp
@@ -264,6 +264,12 @@ FTimespan UMediaPlayer::GetTime() const
 	return PlayerFacade->GetTime();
 }
 
+FTimespan UMediaPlayer::GetDisplayTime() const
+{
+	auto TimeStamp = PlayerFacade->GetDisplayTimeStamp();
+	return TimeStamp.IsValid() ? TimeStamp.Time : FTimespan::Zero();
+}
+
 UMediaTimeStampInfo* UMediaPlayer::GetTimeStamp() const
 {
 	UMediaTimeStampInfo *TimeStampInfo = NewObject<UMediaTimeStampInfo>();
@@ -276,6 +282,18 @@ UMediaTimeStampInfo* UMediaPlayer::GetTimeStamp() const
 	return TimeStampInfo;
 }
 
+UMediaTimeStampInfo* UMediaPlayer::GetDisplayTimeStamp() const
+{
+	UMediaTimeStampInfo* TimeStampInfo = NewObject<UMediaTimeStampInfo>();
+	if (TimeStampInfo)
+	{
+		FMediaTimeStamp TimeStamp = PlayerFacade->GetDisplayTimeStamp();
+		TimeStampInfo->Time = TimeStamp.Time;
+		TimeStampInfo->SequenceIndex = TimeStamp.SequenceIndex;
+	}
+	return TimeStampInfo;
+}
+
 FText UMediaPlayer::GetTrackDisplayName(EMediaPlayerTrack TrackType, int32 TrackIndex) const
 {
 	return PlayerFacade->GetTrackDisplayName((EMediaTrackType)TrackType, TrackIndex);
diff --git a/Engine/Source/Runtime/MediaAssets/Private/Assets/MediaTexture.cpp b/Engine/Source/Runtime/MediaAssets/Private/Assets/MediaTexture.cpp
--- a/Engine/Source/Runtime/MediaAssets/Private/Assets/MediaTexture.cpp
+++ b/Engine/Source/Runtime/MediaAssets/Private/Assets/MediaTexture.cpp
@@ -360,12 +360,8 @@ void UMediaTexture::TickResource(FTimespan Timecode)
 	bool bIsSampleValid = false;
 	if (UMediaPlayer* CurrentPlayerPtr = CurrentPlayer.Get())
 	{
-		const bool PlayerActive = CurrentPlayerPtr->IsPaused() || CurrentPlayerPtr->IsPlaying() || CurrentPlayerPtr->IsPreparing();
-
-		if (PlayerActive)
+		if (CurrentPlayerPtr->GetPlayerFacade()->GetPlayer().IsValid())
 		{
-			check(CurrentPlayerPtr->GetPlayerFacade()->GetPlayer());
-
 			if (CurrentPlayerPtr->GetPlayerFacade()->GetPlayer()->GetPlayerFeatureFlag(IMediaPlayer::EFeatureFlag::UsePlaybackTimingV2))
 			{
 				/*
@@ -404,23 +400,37 @@ void UMediaTexture::TickResource(FTimespan Timecode)
 				//
 				// Old style: pass queue along and dequeue only at render time
 				//
-				TSharedPtr<IMediaTextureSample, ESPMode::ThreadSafe> Sample;
-				if (SampleQueue->Peek(Sample))
+				const bool PlayerActive = CurrentPlayerPtr->IsPaused() || CurrentPlayerPtr->IsPlaying() || CurrentPlayerPtr->IsPreparing();
+				if (PlayerActive)
 				{
-					UpdateSampleInfo(Sample);
+					TSharedPtr<IMediaTextureSample, ESPMode::ThreadSafe> Sample;
+					if (SampleQueue->Peek(Sample))
+					{
+						UpdateSampleInfo(Sample);
 
-					// See above: track sRGB state (for V1 we don't look at all samples - but this should be fine as this should not change on a per sample basis usually)
-					SRGB = Sample->IsOutputSrgb();
-					LastSrgb = SRGB;
+						// See above: track sRGB state (for V1 we don't look at all samples - but this should be fine as this should not change on a per sample basis usually)
+						SRGB = Sample->IsOutputSrgb();
+						LastSrgb = SRGB;
 
-					TextureNumMips = (Sample->GetNumMips() > 1) ? Sample->GetNumMips() : NumMips;
-					bIsSampleValid = true;
-				}
+						TextureNumMips = (Sample->GetNumMips() > 1) ? Sample->GetNumMips() : NumMips;
+						bIsSampleValid = true;
+					}
 
-				RenderParams.SampleSource = SampleQueue;
+					RenderParams.SampleSource = SampleQueue;
 
-				RenderParams.Rate = CurrentPlayerPtr->GetRate();
-				RenderParams.Time = CurrentPlayerPtr->GetTime();
+					RenderParams.Rate = CurrentPlayerPtr->GetRate();
+					RenderParams.Time = CurrentPlayerPtr->GetTime();
+				}
+				else
+				{
+					CurrentAspectRatio = 0.0f;
+					CurrentOrientation = MTORI_Original;
+
+					if (!AutoClear)
+					{
+						return; // retain last frame
+					}
+				}
 			}
 		}
 		else 
diff --git a/Engine/Source/Runtime/MediaAssets/Public/MediaPlayer.h b/Engine/Source/Runtime/MediaAssets/Public/MediaPlayer.h
--- a/Engine/Source/Runtime/MediaAssets/Public/MediaPlayer.h
+++ b/Engine/Source/Runtime/MediaAssets/Public/MediaPlayer.h
@@ -307,6 +307,15 @@ public:
 	UFUNCTION(BlueprintCallable, Category="Media|MediaPlayer")
 	FTimespan GetTime() const;
 
+	/**
+	 * Get the media's current playback time as appropriate for display.
+	 *
+	 * @return Playback time.
+	 * @see GetDuration, Seek
+	 */
+	UFUNCTION(BlueprintCallable, Category = "Media|MediaPlayer")
+	FTimespan GetDisplayTime() const;
+
 	/**
 	 * Get the media's current playback timestamp.
 	 *
@@ -316,6 +325,15 @@ public:
 	UFUNCTION(BlueprintCallable, Category = "Media|MediaPlayer")
 	UMediaTimeStampInfo* GetTimeStamp() const;
 
+	/**
+	 * Get the media's current playback timestamp as appropriate for display.
+	 *
+	 * @return Playback timestamp.
+	 * @see GetDuration, Seek
+	 */
+	UFUNCTION(BlueprintCallable, Category = "Media|MediaPlayer")
+	UMediaTimeStampInfo* GetDisplayTimeStamp() const;
+
 	/**
 	 * Get the human readable name of the specified track.
 	 *
diff --git a/Engine/Source/Runtime/MediaUtils/Private/MediaPlayerFacade.cpp b/Engine/Source/Runtime/MediaUtils/Private/MediaPlayerFacade.cpp
--- a/Engine/Source/Runtime/MediaUtils/Private/MediaPlayerFacade.cpp
+++ b/Engine/Source/Runtime/MediaUtils/Private/MediaPlayerFacade.cpp
@@ -76,6 +76,10 @@ DECLARE_DWORD_ACCUMULATOR_STAT(TEXT("MediaPlayerFacade TotalPurgedVideoSamples")
 DECLARE_DWORD_COUNTER_STAT(TEXT("MediaPlayerFacade NumPurgedSubtitleSamples"), STAT_MediaUtils_FacadeNumPurgedSubtitleSamples, STATGROUP_Media);
 DECLARE_DWORD_ACCUMULATOR_STAT(TEXT("MediaPlayerFacade TotalPurgedSubtitleSamples"), STAT_MediaUtils_FacadeTotalPurgedSubtitleSamples, STATGROUP_Media);
 
+/** Number of purged caption samples */
+DECLARE_DWORD_COUNTER_STAT(TEXT("MediaPlayerFacade NumPurgedCaptionSamples"), STAT_MediaUtils_FacadeNumPurgedCaptionSamples, STATGROUP_Media);
+DECLARE_DWORD_ACCUMULATOR_STAT(TEXT("MediaPlayerFacade TotalPurgedCaptionSamples"), STAT_MediaUtils_FacadeTotalPurgedCaptionSamples, STATGROUP_Media);
+
 /* Some constants
 *****************************************************************************/
 
@@ -86,6 +90,7 @@ static const double kOutdatedSubtitleSamplesTolerance = 0.050;	// seconds subtit
 static const double kOutdatedSamplePurgeRange = 1.0;			// milliseconds for pseudo DT timespan used with async purging of outdated video samples
 static const int32	kMinFramesInVideoQueueToPurge = 3;			// we only consider purging any old frames from the video queue if more than these are present (to not kill a slow playback entirely)
 static const int32	kMinFramesInSubtitleQueueToPurge = 3;		// we only consider purging any old frames from the subtitle queue if more than these are present (to not kill a slow playback entirely)
+static const int32	kMinFramesInCaptionQueueToPurge = 3;		// we only consider purging any old frames from the caption queue if more than these are present (to not kill a slow playback entirely)
 
 /* Local helpers
 *****************************************************************************/
@@ -430,13 +435,13 @@ TRangeSet<float> FMediaPlayerFacade::GetSupportedRates(bool Unthinned) const
 
 bool FMediaPlayerFacade::HaveVideoPlayback() const
 {
-	return VideoSampleSinks.Num() && GetSelectedTrack(EMediaTrackType::Video) != INDEX_NONE;
+	return VideoSampleSinks.Num() && (GetSelectedTrack(EMediaTrackType::Video) != INDEX_NONE);
 }
 
 
 bool FMediaPlayerFacade::HaveAudioPlayback() const
 {
-	return PrimaryAudioSink.IsValid() && GetSelectedTrack(EMediaTrackType::Audio) != INDEX_NONE;
+	return PrimaryAudioSink.IsValid() && (GetSelectedTrack(EMediaTrackType::Audio) != INDEX_NONE);
 }
 
 
@@ -472,6 +477,17 @@ FTimespan FMediaPlayerFacade::GetTime() const
 
 
 FMediaTimeStamp FMediaPlayerFacade::GetTimeStamp() const
+{
+	return GetTimeStampInternal(false);
+}
+
+
+FMediaTimeStamp FMediaPlayerFacade::GetDisplayTimeStamp() const
+{
+	return GetTimeStampInternal(true);
+}
+
+FMediaTimeStamp FMediaPlayerFacade::GetTimeStampInternal(bool bForDisplay) const
 {
 	TSharedPtr<IMediaPlayer, ESPMode::ThreadSafe> CurrentPlayer(Player);
 	if (!CurrentPlayer.IsValid())
@@ -487,6 +503,12 @@ FMediaTimeStamp FMediaPlayerFacade::GetTimeStamp() const
 		return FMediaTimeStamp(GetTime());
 	}
 
+	// Check if the value is for display purposes. If so: do we seek right now?
+	if (bForDisplay && SeekTargetTime.IsValid())
+	{
+		return SeekTargetTime;
+	}
+
 	// Check if there are video samples present or presence is unknown.
 	// Only when we know for sure that there are none because the existing video stream has ended do we set this to false.
 	bool bHaveVideoSamples = VideoSampleAvailability != 0;
@@ -496,7 +518,7 @@ FMediaTimeStamp FMediaPlayerFacade::GetTimeStamp() const
 		/*
 			Returning the precise time of the sample returned during TickFetch()
 		*/
-		return CurrentFrameVideoTimeStamp;
+		return bForDisplay ? CurrentFrameVideoDisplayTimeStamp : CurrentFrameVideoTimeStamp;
 	}
 	else if (HaveAudioPlayback())
 	{
@@ -1080,13 +1102,16 @@ bool FMediaPlayerFacade::ContinueOpen(IMediaPlayerLifecycleManagerDelegate::ICon
 		FScopeLock Lock(&LastTimeValuesCS);
 
 		BlockOnRangeDisabled = false;
-		BlockOnRange.Flush();
+		BlockOnRange.OnFlush();
 		LastVideoSampleProcessedTimeRange = TRange<FMediaTimeStamp>::Empty();
 		LastAudioSampleProcessedTime.Invalidate();
 		CurrentFrameVideoTimeStamp.Invalidate();
+		CurrentFrameVideoDisplayTimeStamp.Invalidate();
 		CurrentFrameAudioTimeStamp.Invalidate();
 
 		NextEstVideoTimeAtFrameStart.Invalidate();
+		SeekTargetTime.Invalidate();
+		SeekIndex = 0;
 	}
 
 	if (bCreateNewPlayer)
@@ -1128,16 +1153,34 @@ void FMediaPlayerFacade::QueryCacheState(EMediaTrackType TrackType, EMediaCacheS
 
 bool FMediaPlayerFacade::Seek(const FTimespan& Time)
 {
-	if (!Player.IsValid() || !Player->GetControls().Seek(Time))
+	auto CurrentPlayer = Player;
+
+	if (!CurrentPlayer.IsValid() || !CurrentPlayer->GetControls().Seek(Time))
 	{
 		return false;
 	}
 
-	if (Player.IsValid() && Player->FlushOnSeekStarted())
+	FScopeLock Lock(&CriticalSection);
+
+	// V2 timing player?
+	if (CurrentPlayer->GetPlayerFeatureFlag(IMediaPlayer::EFeatureFlag::UsePlaybackTimingV2))
+	{
+		// Yes. Flush only the facade side of the system as needed for seeks
+		// (the player is expected to flush its internal queues as needed itself)
+		check(CurrentPlayer->GetPlayerFeatureFlag(IMediaPlayer::EFeatureFlag::PlayerUsesInternalFlushOnSeek));
+		Flush(true, true);
+	}
+	else
 	{
-		Flush(Player->GetPlayerFeatureFlag(IMediaPlayer::EFeatureFlag::PlayerUsesInternalFlushOnSeek));
+		// No. Flush as requested...
+		if (CurrentPlayer->FlushOnSeekStarted())
+		{
+			Flush(CurrentPlayer->GetPlayerFeatureFlag(IMediaPlayer::EFeatureFlag::PlayerUsesInternalFlushOnSeek), false);
+		}
 	}
 
+	SeekTargetTime = FMediaTimeStamp(Time, FMediaTimeStamp::MakeSequenceIndex(SeekIndex, 0));
+
 	return true;
 }
 
@@ -1189,17 +1232,32 @@ void FMediaPlayerFacade::SetBlockOnTimeRange(const TRange<FTimespan>& TimeRange)
 }
 
 
-void FMediaPlayerFacade::FBlockOnRange::Flush()
+void FMediaPlayerFacade::FBlockOnRange::OnFlush()
+{
+	LastBlockOnRange = TRange<FTimespan>::Empty();
+	OnBlockPrimaryIndex = 0;
+	OnBlockSecondaryIndex = 0;
+}
+
+
+void FMediaPlayerFacade::FBlockOnRange::OnSeek(int32 PrimaryIndex)
 {
+	TSharedPtr<IMediaPlayer, ESPMode::ThreadSafe> CurrentPlayer(Facade->Player);
+	check(CurrentPlayer.IsValid());
+
 	LastBlockOnRange = TRange<FTimespan>::Empty();
-	OnBlockSeqIndex = 0;
+	OnBlockPrimaryIndex = PrimaryIndex;
+	OnBlockSecondaryIndex = 0;
 }
 
 
 void FMediaPlayerFacade::FBlockOnRange::SetRange(const TRange<FTimespan>& NewRange)
 {
-	CurrentTimeRange = NewRange;
-	RangeIsDirty = true;
+	if (CurrentTimeRange != NewRange)
+	{
+		CurrentTimeRange = NewRange;
+		RangeIsDirty = true;
+	}
 }
 
 
@@ -1226,6 +1284,7 @@ const TRange<FMediaTimeStamp>& FMediaPlayerFacade::FBlockOnRange::GetRange() con
 		return BlockOnRange;
 	}
 
+	// If the range is empty or the player can't support blocked playback: reset everything & return empty block range...
 	if (CurrentTimeRange.IsEmpty() || !CurrentPlayer->GetControls().CanControl(EMediaControl::BlockOnFetch))
 	{
 		LastBlockOnRange = TRange<FTimespan>::Empty();
@@ -1249,10 +1308,20 @@ const TRange<FMediaTimeStamp>& FMediaPlayerFacade::FBlockOnRange::GetRange() con
 
 	auto SetBlockOnRange = BlockOnRange;
 
+	/*
+	* On the synthesized sequence and loop index values:
+	* - We track seeks and hence can insert the proper seek index easily, although the user does not provide it / does not need to track it
+	* - The loop index gets somewhat of a special treatment:
+	*  -- With tools like Sequencer a blocked range my speed along the timeline quite quickly and if the player is configured as looping it would be expected that the video loops while this is done
+	*  -- We hence could loop multiple times within a single update interval
+	*  -- Still we treat any loop (aka: a jump "backwards" without an explicit seek) as a single loop iteration
+	*     (this is easier for any player to work with and provides the same visual results)
+	*/
+
 	if (!CurrentPlayer->GetControls().IsLooping())
 	{
-		// We pass in the time range as is on seq-index zero at all times - players have to reject sample output / blocking logic will detect begin outside media range
-		BlockOnRange = TRange<FMediaTimeStamp>(FMediaTimeStamp(Start, 0), FMediaTimeStamp(End, 0));
+		int64 SequenceIndex = FMediaTimeStamp::MakeSequenceIndex(OnBlockPrimaryIndex, OnBlockSecondaryIndex);
+		BlockOnRange = TRange<FMediaTimeStamp>(FMediaTimeStamp(Start, SequenceIndex), FMediaTimeStamp(End, SequenceIndex));
 	}
 	else
 	{
@@ -1279,8 +1348,16 @@ const TRange<FMediaTimeStamp>& FMediaPlayerFacade::FBlockOnRange::GetRange() con
 		Start = WrappedModulo(Start, Duration);
 		End = WrappedModulo(End, Duration);
 
+		/*
+		* Note about moving along over multiple loops at once:
+		* We do not explicitly care about this case, but rather look at the blocking range to always be in the "first" loop, even if - on the external time line - we might be on a higher loop count.
+		* Hence we only need to detect if we moved "backwards" without an explicit seek vs. the last requested blocked range (see below) and then increase the loop index by one.
+		* 
+		* -> visually this will result in the desired behavior
+		*/
 
 		// Detect any non-monotonic movement of the range...
+		// (seeks will reset the last value, so this is ensured not to trigger in that case)
 		if (!LastBlockOnRange.IsEmpty())
 		{
 			if (Rate >= 0.0f)
@@ -1288,7 +1365,7 @@ const TRange<FMediaTimeStamp>& FMediaPlayerFacade::FBlockOnRange::GetRange() con
 				FTimespan LastStart = WrappedModulo(LastBlockOnRange.GetLowerBoundValue(), Duration);
 				if ((LastStart > Start) || ((LastStart == Start) && (LastBlockOnRange.GetLowerBoundValue() < CurrentTimeRange.GetLowerBoundValue())))
 				{
-					++OnBlockSeqIndex;
+					++OnBlockSecondaryIndex;
 				}
 			}
 			else
@@ -1296,26 +1373,28 @@ const TRange<FMediaTimeStamp>& FMediaPlayerFacade::FBlockOnRange::GetRange() con
 				FTimespan LastEnd = WrappedModulo(LastBlockOnRange.GetUpperBoundValue(), Duration);
 				if ((LastEnd < End) || ((LastEnd == End) && (LastBlockOnRange.GetUpperBoundValue() > CurrentTimeRange.GetUpperBoundValue())))
 				{
-					--OnBlockSeqIndex;
+					--OnBlockSecondaryIndex;
 				}
 			}
 		}
 
 		// Check if our range crosses the loop point (sequence boundary)
-		int64 StartIndex, EndIndex;
+		int32 StartLoopIndex, EndLoopIndex;
 		if (Rate >= 0.0)
 		{
-			StartIndex = OnBlockSeqIndex;
-			EndIndex = (Start <= End) ? OnBlockSeqIndex : (OnBlockSeqIndex + 1);
+			StartLoopIndex = OnBlockSecondaryIndex;
+			EndLoopIndex = (Start <= End) ? OnBlockSecondaryIndex : (OnBlockSecondaryIndex + 1);
 		}
 		else
 		{
-			StartIndex = (Start <= End) ? OnBlockSeqIndex : (OnBlockSeqIndex - 1);
-			EndIndex = OnBlockSeqIndex;
+			StartLoopIndex = (Start <= End) ? OnBlockSecondaryIndex : (OnBlockSecondaryIndex - 1);
+			EndLoopIndex = OnBlockSecondaryIndex;
 		}
 
 		// Assemble final blocking range
-		BlockOnRange = TRange<FMediaTimeStamp>(FMediaTimeStamp(Start, StartIndex), FMediaTimeStamp(End, EndIndex));
+		auto SeqIndexStart = FMediaTimeStamp::MakeSequenceIndex(OnBlockPrimaryIndex, StartLoopIndex);
+		auto SeqIndexEnd = FMediaTimeStamp::MakeSequenceIndex(OnBlockPrimaryIndex, EndLoopIndex);
+		BlockOnRange = TRange<FMediaTimeStamp>(FMediaTimeStamp(Start, SeqIndexStart), FMediaTimeStamp(End, SeqIndexEnd));
 		check(!BlockOnRange.IsEmpty());
 	}
 
@@ -1371,22 +1450,33 @@ bool FMediaPlayerFacade::SetRate(float Rate)
 	// Enter CS as we change the rate which we read on the tickable thread
 	FScopeLock Lock(&CriticalSection);
 
-	// Can we set the rate at all?
-	if (!Player.IsValid() || !Player->GetControls().SetRate(Rate))
+	if (!Player.IsValid())
 	{
 		return false;
 	}
 
 	// Is this new rate supported?
+	bool bRateOk = true;
 	if (Rate != 0.0f && !(Player->GetControls().GetSupportedRates(EMediaRateThinning::Thinned).Contains(Rate) || Player->GetControls().GetSupportedRates(EMediaRateThinning::Unthinned).Contains(Rate)))
+	{
+		// Pause player instead...
+		// (some players may do this as a reaction to the illegal rate anyways - but we need to track the state properly!)
+		Rate = 0.0f;
+		bRateOk = false;
+	}
+
+	// Attempt to set the rate...
+	if (!Player->GetControls().SetRate(Rate))
 	{
 		return false;
 	}
 
+
+	// Any change?
 	if (CurrentRate == Rate)
 	{
 		// no change - just return with ok status
-		return true;
+		return bRateOk;
 	}
 
 	if ((LastRate * Rate) < 0.0f)
@@ -1405,13 +1495,14 @@ bool FMediaPlayerFacade::SetRate(float Rate)
 		}
 	}
 
+	// Track last "unpaused" rate we set
 	if (Rate != 0.0)
 	{
 		LastRate = Rate;
 	}
 	CurrentRate = Rate;
 
-	return true;
+	return bRateOk;
 }
 
 
@@ -1494,56 +1585,8 @@ bool FMediaPlayerFacade::BlockOnFetch() const
 		// V2 blocking logic
 		//
 
-		float Rate = GetUnpausedRate();
-
-		// If the current sample "out there" is actually overlapping with the current block, we might be good with no new sample
-		if (LastVideoSampleProcessedTimeRange.Overlaps(BR))
-		{
-			// We have no new data (else we would not even call this method), but the last sample we returned is still inside the current range -> good, but...
-			//
-			// If the next sample would already cover more of the range than the older one we would like to use that instead -> but it may well be we do not have any data about the sample yet (and would indeed LIKE to block!)
-			// So, we assume that the next sample will follow with no gap and have the same duration (a pretty good, general assumption) and check against that data to see if it would be better...
-
-			// Get last sample's time range
-			TRange<FMediaTimeStamp> LastSampleTimeRange(LastVideoSampleProcessedTimeRange);
-
-			// Compute the "theoretical" next sample range...
-			TRange<FMediaTimeStamp> NextSampleTimeRange = (Rate >= 0.0f) ? TRange<FMediaTimeStamp>(LastVideoSampleProcessedTimeRange.GetUpperBoundValue(), LastVideoSampleProcessedTimeRange.GetUpperBoundValue() + LastVideoSampleProcessedTimeRange.Size<FMediaTimeStamp>().Time)
-				: TRange<FMediaTimeStamp>(LastVideoSampleProcessedTimeRange.GetLowerBoundValue() - LastVideoSampleProcessedTimeRange.Size<FMediaTimeStamp>().Time, LastVideoSampleProcessedTimeRange.GetLowerBoundValue());
-
-			FTimespan Duration = Player->GetControls().GetDuration();
-
-			if (!Player->GetControls().IsLooping())
-			{
-				// If we are not looping we need to clamp against the media's duration
-				// (we assume it starts at zero here!)
-				check(NextSampleTimeRange.GetLowerBoundValue().SequenceIndex == 0);
-				NextSampleTimeRange = TRange<FMediaTimeStamp>::Intersection(NextSampleTimeRange, TRange<FMediaTimeStamp>(FMediaTimeStamp(0, 0), FMediaTimeStamp(Duration, 0)));
-			}
-			else
-			{
-				if (NextSampleTimeRange.GetLowerBoundValue().Time >= Duration)
-				{
-					check(Rate >= 0.0f);
-					NextSampleTimeRange = TRange<FMediaTimeStamp>(FMediaTimeStamp(NextSampleTimeRange.GetLowerBoundValue().Time - Duration, NextSampleTimeRange.GetLowerBoundValue().SequenceIndex + 1), FMediaTimeStamp(NextSampleTimeRange.GetUpperBoundValue().Time - Duration, NextSampleTimeRange.GetUpperBoundValue().SequenceIndex + 1));
-				}
-				else if (NextSampleTimeRange.GetLowerBoundValue().Time < FTimespan::Zero())
-				{
-					check(Rate < 0.0f);
-					NextSampleTimeRange = TRange<FMediaTimeStamp>(FMediaTimeStamp(NextSampleTimeRange.GetLowerBoundValue().Time + Duration, NextSampleTimeRange.GetLowerBoundValue().SequenceIndex - 1), FMediaTimeStamp(NextSampleTimeRange.GetUpperBoundValue().Time + Duration, NextSampleTimeRange.GetUpperBoundValue().SequenceIndex - 1));
-				}
-			}
-
-			// Compute which one is larger inside the current range...
-			int64 LastSampleCoverage = TRange<FMediaTimeStamp>::Intersection(BR, LastSampleTimeRange).Size<FMediaTimeStamp>().Time.GetTicks();
-			int64 NextSampleCoverage = TRange<FMediaTimeStamp>::Intersection(BR, NextSampleTimeRange).Size<FMediaTimeStamp>().Time.GetTicks();
-
-			if (LastSampleCoverage > NextSampleCoverage)
-			{
-				// Last one we returned is still good. No blocking needed...
-				return false;
-			}
-		}
+		// note: with V2 timing we only get here if any current sample is no longer considered "valid" and we didn't so far get a new one that would be
+		//  -->  we do not need to check the actual range here; we only check for exceptions, where we can proceed although we don't have the sample...
 
 		// The next checks make only sense if the player is done preparing...
 		if (!IsPreparing())
@@ -1553,8 +1596,24 @@ bool FMediaPlayerFacade::BlockOnFetch() const
 			{
 				// Yes. Is the sample outside the media's range?
 				// (note: this assumes the media starts at time ZERO - this will not be the case at all times (e.g. life playback) -- for now we assume a player will flagged blocked playback as invalid in that case!)
-				if (BR.GetUpperBoundValue() < FMediaTimeStamp(0) || Player->GetControls().GetDuration() <= BR.GetLowerBoundValue().Time)
+				if (BR.GetUpperBoundValue() < FMediaTimeStamp(FTimespan::Zero(), BR.GetUpperBoundValue().SequenceIndex) || Player->GetControls().GetDuration() <= BR.GetLowerBoundValue().Time)
+				{
+					return false;
+				}
+			}
+			// Paused?
+			if (CurrentRate == 0.0f)
+			{
+				// Yes. If we get here the current sample did not satisfy the set range. So only a incoming seek could resolved this...
+				if (!SeekTargetTime.IsValid())
+				{
+					// No, so don't block. This range is impossible to satisfy...
+					return false;
+				}
+				// Would the seek target probably fit the range?
+				if (!BR.Overlaps(TRange<FMediaTimeStamp>(SeekTargetTime, SeekTargetTime + BR.Size<FMediaTimeStamp>().Time)))
 				{
+					// No, so don't block. This range is impossible to satisfy...
 					return false;
 				}
 			}
@@ -1598,7 +1657,7 @@ bool FMediaPlayerFacade::BlockOnFetch() const
 }
 
 
-void FMediaPlayerFacade::Flush(bool bExcludePlayer)
+void FMediaPlayerFacade::Flush(bool bExcludePlayer, bool bOnSeek)
 {
 	UE_LOG(LogMediaUtils, Verbose, TEXT("PlayerFacade %p: Flushing sinks"), this);
 
@@ -1615,13 +1674,29 @@ void FMediaPlayerFacade::Flush(bool bExcludePlayer)
 		Player->GetSamples().FlushSamples();
 	}
 
-	NextVideoSampleTime = FTimespan::MinValue();
-
 	LastAudioRenderedSampleTime.Invalidate();
+	if (bOnSeek)
+	{
+		SeekIndex += (GetUnpausedRate() < 0.0f) ? -1 : 1;
+		BlockOnRange.OnSeek(SeekIndex);
+	}
+	else
+	{
+		BlockOnRange.OnFlush();
+		SeekIndex = 0;
+	}
+
+	// Logically we have no old sample anymore
+	// (as in: we will start asking for a new one until we get one - even with a rate of zero, if we had a non-zero one ever before)
 	LastVideoSampleProcessedTimeRange = TRange<FMediaTimeStamp>::Empty();
-	BlockOnRange.Flush();
 
+	// Invalidate next video time to fetch (none-audio case)
 	NextEstVideoTimeAtFrameStart.Invalidate();
+	// ...and seek target
+	SeekTargetTime.Invalidate();
+
+	// V1 only
+	NextVideoSampleTime = FTimespan::MinValue();
 }
 
 
@@ -1798,13 +1873,22 @@ void FMediaPlayerFacade::ProcessEvent(EMediaEvent Event, bool bIsBroadcastAllowe
 	if ((Event == EMediaEvent::PlaybackEndReached) ||
 		(Event == EMediaEvent::TracksChanged))
 	{
-		Flush();
+		if (Player.IsValid() && !Player->GetPlayerFeatureFlag(IMediaPlayer::EFeatureFlag::UsePlaybackTimingV2))
+		{
+			// Execute flush for older players only
+			Flush();
+		}
 	}
 	else if (Event == EMediaEvent::SeekCompleted)
 	{
-		if (!Player.IsValid() || Player->FlushOnSeekCompleted())
+		// We only consider flushing on seek completion if there is a V1 timing player...
+		if (Player.IsValid() && !Player->GetPlayerFeatureFlag(IMediaPlayer::EFeatureFlag::UsePlaybackTimingV2))
 		{
-			Flush(!Player.IsValid() || Player->GetPlayerFeatureFlag(IMediaPlayer::EFeatureFlag::PlayerUsesInternalFlushOnSeek));
+			// Does the player want this?
+			if (Player->FlushOnSeekCompleted())
+			{
+				Flush(Player->GetPlayerFeatureFlag(IMediaPlayer::EFeatureFlag::PlayerUsesInternalFlushOnSeek), true);
+			}
 		}
 	}
 	else if (Event == EMediaEvent::MediaClosed)
@@ -1875,9 +1959,7 @@ void FMediaPlayerFacade::TickInput(FTimespan DeltaTime, FTimespan Timecode)
 
 	if (Player.IsValid())
 	{
-		// Update flag reflecting presence of audio in the current stream
-		// (doing it just once per gameloop is enough)
-		bHaveActiveAudio = HaveAudioPlayback();
+		MonitorAudioEnablement();
 
 		Player->TickInput(DeltaTime, Timecode);
 
@@ -1942,10 +2024,11 @@ void FMediaPlayerFacade::TickInput(FTimespan DeltaTime, FTimespan Timecode)
 
 				if (ProcessVideoSamples(Samples, TimeRange))
 				{
+					// We either got a new sample or a current one is still the best choice...
 					break;
 				}
 
-				// No sample. Should we block for one?
+				// The current one is outdated and no new one was delivered. Should we block for one?
 				if (!BlockOnFetch())
 				{
 					// No... continue...
@@ -1955,6 +2038,9 @@ void FMediaPlayerFacade::TickInput(FTimespan DeltaTime, FTimespan Timecode)
 				// Issue tick call with dummy timing as some players advance some state in the tick, which we wait for
 				Player->TickInput(FTimespan::Zero(), FTimespan::MinValue());
 
+				// Monitor / update seek status
+				UpdateSeekStatus();
+
 				// Process deferred events & check for events that break the block
 				bool bEventCancelsBlock = false;
 				while (QueuedEvents.Dequeue(Event))
@@ -2086,8 +2172,8 @@ void FMediaPlayerFacade::TickFetch(FTimespan DeltaTime, FTimespan Timecode)
 
 		while (true)
 		{
-			ProcessCaptionSamples(Samples, TimeRange);
-			ProcessSubtitleSamples(Samples, TimeRange);
+			ProcessCaptionSamplesV1(Samples, TimeRange);
+			ProcessSubtitleSamplesV1(Samples, TimeRange);
 			ProcessVideoSamplesV1(Samples, TimeRange);
 
 			if (!BlockOnFetch())
@@ -2163,42 +2249,154 @@ void FMediaPlayerFacade::TickTickable()
 	Player->TickAudio();
 
 	// determine range of valid samples
-	TRange<FTimespan> AudioTimeRange;
-	TRange<FTimespan> MetadataTimeRange;
 
-	const FTimespan Time = GetTime();
+	// process samples in range
+	IMediaSamples& Samples = Player->GetSamples();
+
+	if (Player->GetPlayerFeatureFlag(IMediaPlayer::EFeatureFlag::UsePlaybackTimingV2))
+	{
+		const FMediaTimeStamp Time = GetTimeStamp();
 
-	bool bUseV2Timing = Player->GetPlayerFeatureFlag(IMediaPlayer::EFeatureFlag::UsePlaybackTimingV2);
+		auto TimeRange = TRange<FMediaTimeStamp>::Inclusive(FMediaTimeStamp(FTimespan::MinValue(), MIN_int32), Time + MediaPlayerFacade::MetadataPreroll);
 
-	if (Rate > 0.0f)
+		ProcessAudioSamples(Samples, TRange<FMediaTimeStamp>());
+		ProcessMetadataSamples(Samples, TimeRange);
+	}
+	else
 	{
-		if (!bUseV2Timing) // we leave range open - sends all the player has
+		TRange<FTimespan> AudioTimeRange;
+		TRange<FTimespan> MetadataTimeRange;
+
+		const FTimespan Time = GetTime();
+
+		if (Rate >= 0.0f)
 		{
 			AudioTimeRange = TRange<FTimespan>::Inclusive(FTimespan::MinValue(), Time + MediaPlayerFacade::AudioPreroll);
+			MetadataTimeRange = TRange<FTimespan>::Inclusive(FTimespan::MinValue(), Time + MediaPlayerFacade::MetadataPreroll);
+		}
+		else
+		{
+			AudioTimeRange = TRange<FTimespan>::Inclusive(Time - MediaPlayerFacade::AudioPreroll, FTimespan::MaxValue());
+			MetadataTimeRange = TRange<FTimespan>::Inclusive(Time - MediaPlayerFacade::MetadataPreroll, FTimespan::MaxValue());
+		}
+
+		ProcessAudioSamplesV1(Samples, AudioTimeRange);
+		ProcessMetadataSamplesV1(Samples, MetadataTimeRange);
+	}
+
+	SET_DWORD_STAT(STAT_MediaUtils_FacadeNumAudioSamples, Samples.NumAudio());
+}
+
+
+void FMediaPlayerFacade::UpdateSeekStatus(const FMediaTimeStamp* pCheckTimeStamp)
+{
+	check(Player.IsValid());
+
+	FScopeLock Lock(&CriticalSection);
+
+	if (HaveVideoPlayback())
+	{
+		if (SeekTargetTime.IsValid())
+		{
+			// Either peek for the newest available sample or take a given timestamp to check against
+			FMediaTimeStamp VideoTimeStamp;
+			if (pCheckTimeStamp)
+			{
+				VideoTimeStamp = *pCheckTimeStamp;
+			}
+			else
+			{
+				Player->GetSamples().PeekVideoSampleTime(VideoTimeStamp);
+			}
+
+			if (VideoTimeStamp.IsValid())
+			{
+				bool bRunningNonAudioClock = bHaveActiveAudio && !BlockOnRange.IsSet();
+
+				if (GetUnpausedRate() >= 0.0f)
+				{
+					// See if we already are looking at a sample from the target sequence index...
+					// (we are not checking for the precise location as some players might not able to deliver it)
+					if (FMediaTimeStamp::GetPrimaryIndex(VideoTimeStamp.SequenceIndex) < FMediaTimeStamp::GetPrimaryIndex(SeekTargetTime.SequenceIndex))
+					{
+						// No. Make sure we drop the sample & possible more up to the seek target (we use the fetch code to avoid any races with a old-sample-purge logic operating async)
+						Player->GetSamples().DiscardVideoSamples(TRange<FMediaTimeStamp>(VideoTimeStamp, SeekTargetTime), false);
+					}
+					else
+					{
+						// We have reached the sequence of the seek target, reset everything to normal operation...
+						// (we do not care if we reached the precise location beyond the index as we do not know if we even can)
+						if (bRunningNonAudioClock)
+						{
+							NextEstVideoTimeAtFrameStart = FMediaTimeStampSample(VideoTimeStamp, FPlatformTime::Seconds());
+						}
+
+						FScopeLock LockLT(&LastTimeValuesCS);
+
+						// Update the display version of then "current frame time" right now (to avoid any glitches as it can take a little while for the frame to actually change)
+						CurrentFrameVideoDisplayTimeStamp = SeekTargetTime;
+
+						// Seeking done
+						SeekTargetTime.Invalidate();
+
+					}
+				}
+				else
+				{
+					// See if we already are looking at a sample from the target sequence index...
+					// (we are not checking for the precise location as some players might not able to deliver it)
+					if (FMediaTimeStamp::GetPrimaryIndex(VideoTimeStamp.SequenceIndex) > FMediaTimeStamp::GetPrimaryIndex(SeekTargetTime.SequenceIndex))
+					{
+						// No. Make sure we drop the sample & possible more up to the seek target (we use the fetch code to avoid any races with a old-sample-purge logic operating async)
+						Player->GetSamples().DiscardVideoSamples(TRange<FMediaTimeStamp>(VideoTimeStamp, SeekTargetTime), true);
+					}
+					else
+					{
+						// We have reached the sequence of the seek target, reset everything to normal operation...
+						// (we do not care if we reached the precise location beyond the index as we do not know if we even can)
+						if (bRunningNonAudioClock)
+						{
+							NextEstVideoTimeAtFrameStart = FMediaTimeStampSample(VideoTimeStamp, FPlatformTime::Seconds());
+						}
+
+						FScopeLock LockLT(&LastTimeValuesCS);
+						SeekTargetTime.Invalidate();
+					}
+				}
+			}
 		}
-		MetadataTimeRange = TRange<FTimespan>::Inclusive(FTimespan::MinValue(), Time + MediaPlayerFacade::MetadataPreroll);
 	}
 	else
 	{
-		if (!bUseV2Timing) // we leave range open - sends all the player has
+		if (bHaveActiveAudio)
 		{
-			AudioTimeRange = TRange<FTimespan>::Inclusive(Time - MediaPlayerFacade::AudioPreroll, FTimespan::MaxValue());
+			FScopeLock LockLT(&LastTimeValuesCS);
+			if (CurrentFrameAudioTimeStamp >= SeekTargetTime)
+			{
+				SeekTargetTime.Invalidate();
+			}
+		}
+		else
+		{
+			// Neither audio nor video are presently active. We just assume we reached the seek target and continue...
+			// (we currently have no other source of a current sample timestamp)
+			SeekTargetTime.Invalidate();
 		}
-		MetadataTimeRange = TRange<FTimespan>::Inclusive(Time - MediaPlayerFacade::MetadataPreroll, FTimespan::MaxValue());
 	}
+}
 
-	// process samples in range
-	IMediaSamples& Samples = Player->GetSamples();
 
-	ProcessAudioSamples(Samples, AudioTimeRange);
-	ProcessMetadataSamples(Samples, MetadataTimeRange);
-	if (bUseV2Timing)
+void FMediaPlayerFacade::MonitorAudioEnablement()
+{
+	// Update flag reflecting presence of audio in the current stream
+	// (doing it just once per gameloop is enough)
+	bool bHadActiveAudio = bHaveActiveAudio;
+	bHaveActiveAudio = HaveAudioPlayback();
+	if (bHadActiveAudio && !bHaveActiveAudio)
 	{
-		ProcessCaptionSamples(Samples, MetadataTimeRange);
-		ProcessSubtitleSamples(Samples, MetadataTimeRange);
+		// Reset state for dt-based playback so we grab a new PTS value immediately
+		NextEstVideoTimeAtFrameStart.Invalidate();
 	}
-
-	SET_DWORD_STAT(STAT_MediaUtils_FacadeNumAudioSamples, Samples.NumAudio());
 }
 
 
@@ -2206,21 +2404,19 @@ void FMediaPlayerFacade::PreSampleProcessingTimeHandling()
 {
 	check(Player.IsValid());
 
-	// No Audio clock?
-	if (!bHaveActiveAudio)
+	FScopeLock Lock(&CriticalSection);
+
+	UpdateSeekStatus();
+
+	if (!SeekTargetTime.IsValid())
 	{
-		// No external clock? (blocking)
-		if (!BlockOnRange.IsSet())
+		// No seek pending. Can we / Do we need to prime a non-audio clock?
+		if (!bHaveActiveAudio && !BlockOnRange.IsSet() && !NextEstVideoTimeAtFrameStart.IsValid())
 		{
-			// Do we have a current timestamp estimation?
-			if (!NextEstVideoTimeAtFrameStart.IsValid())
+			FMediaTimeStamp VideoTimeStamp;
+			if (Player->GetSamples().PeekVideoSampleTime(VideoTimeStamp))
 			{
-				// Not, yet. We need to attempt to get the next video sample's timestamp to get going...
-				FMediaTimeStamp VideoTimeStamp;
-				if (Player->GetSamples().PeekVideoSampleTime(VideoTimeStamp))
-				{
-					NextEstVideoTimeAtFrameStart = FMediaTimeStampSample(VideoTimeStamp, FPlatformTime::Seconds());
-				}
+				NextEstVideoTimeAtFrameStart = FMediaTimeStampSample(VideoTimeStamp, FPlatformTime::Seconds());
 			}
 		}
 	}
@@ -2241,7 +2437,7 @@ void FMediaPlayerFacade::PostSampleProcessingTimeHandling(FTimespan DeltaTime)
 		{
 			// Move video frame start estimate forward
 			// (the initial NextEstVideoTimeAtFrameStart will never be valid if no video is present)
-			if (!bHaveActiveAudio && NextEstVideoTimeAtFrameStart.IsValid())
+			if (NextEstVideoTimeAtFrameStart.IsValid())
 			{
 				if (Player->GetPlayerFeatureFlag(IMediaPlayer::EFeatureFlag::UseRealtimeWithVideoOnly))
 				{
@@ -2262,18 +2458,18 @@ void FMediaPlayerFacade::PostSampleProcessingTimeHandling(FTimespan DeltaTime)
 				{
 					if (Rate >= 0.0f)
 					{
-						if (NextEstVideoTimeAtFrameStart.TimeStamp.Time >= Duration)
+						while (NextEstVideoTimeAtFrameStart.TimeStamp.Time >= Duration)
 						{
 							NextEstVideoTimeAtFrameStart.TimeStamp.Time -= Duration;
-							++NextEstVideoTimeAtFrameStart.TimeStamp.SequenceIndex;
+							NextEstVideoTimeAtFrameStart.TimeStamp.SequenceIndex = FMediaTimeStamp::AdjustSecondaryIndex(NextEstVideoTimeAtFrameStart.TimeStamp.SequenceIndex, 1);
 						}
 					}
 					else
 					{
-						if (NextEstVideoTimeAtFrameStart.TimeStamp.Time < FTimespan::Zero())
+						while (NextEstVideoTimeAtFrameStart.TimeStamp.Time < FTimespan::Zero())
 						{
 							NextEstVideoTimeAtFrameStart.TimeStamp.Time += Duration;
-							--NextEstVideoTimeAtFrameStart.TimeStamp.SequenceIndex;
+							NextEstVideoTimeAtFrameStart.TimeStamp.SequenceIndex = FMediaTimeStamp::AdjustSecondaryIndex(NextEstVideoTimeAtFrameStart.TimeStamp.SequenceIndex, -1);
 						}
 					}
 				}
@@ -2300,8 +2496,12 @@ void FMediaPlayerFacade::PostSampleProcessingTimeHandling(FTimespan DeltaTime)
 }
 
 
-bool FMediaPlayerFacade::GetCurrentPlaybackTimeRange(TRange<FMediaTimeStamp>& TimeRange, float Rate, FTimespan DeltaTime, bool bDoNotUseFrameStartReference) const
+bool FMediaPlayerFacade::GetCurrentPlaybackTimeRange(TRange<FMediaTimeStamp>& TimeRange, float Rate, FTimespan DeltaTime, bool bPurgeSampleRelated) const
 {
+	/*
+	* Note: while a seek operation is still in progress (no sample from target location has been processed) this will
+	* return on an empty time range.
+	*/
 	check(Player->GetPlayerFeatureFlag(IMediaPlayer::EFeatureFlag::UsePlaybackTimingV2));
 
 	TSharedPtr<FMediaAudioSampleSink, ESPMode::ThreadSafe> AudioSink = PrimaryAudioSink.Pin();
@@ -2315,7 +2515,26 @@ bool FMediaPlayerFacade::GetCurrentPlaybackTimeRange(TRange<FMediaTimeStamp>& Ti
 		FMediaTimeStampSample AudioTime = AudioSink->GetAudioTime();
 		if (!AudioTime.IsValid())
 		{
-			// No timing info available, no timerange available, no samples to process
+			if (!bPurgeSampleRelated)
+			{
+				// If paused and not seeking, make sure we get one sample nonetheless...
+				if (Rate == 0.0f && !SeekTargetTime.IsValid())
+				{
+					// Do this once after open / seek...
+					if (LastVideoSampleProcessedTimeRange.IsEmpty())
+					{
+						// Use the video sample timestamp for simplicity (although we otherwise sync with audio timestamps)
+						FMediaTimeStamp TimeStamp;
+						if (Player->GetSamples().PeekVideoSampleTime(TimeStamp))
+						{
+							TimeRange = TRange<FMediaTimeStamp>(TimeStamp, TimeStamp + DeltaTime);
+							return !TimeRange.IsEmpty();
+						}
+					}
+				}
+			}
+
+			// No timing info available, no time range available, no samples to process
 			return false;
 		}
 
@@ -2323,7 +2542,7 @@ bool FMediaPlayerFacade::GetCurrentPlaybackTimeRange(TRange<FMediaTimeStamp>& Ti
 
 		double Now = FPlatformTime::Seconds();
 
-		if (!bDoNotUseFrameStartReference)
+		if (!bPurgeSampleRelated)
 		{
 			// Normal estimation relative to current frame start...
 			// (on gamethread operation)
@@ -2363,8 +2582,7 @@ bool FMediaPlayerFacade::GetCurrentPlaybackTimeRange(TRange<FMediaTimeStamp>& Ti
 			}
 		}
 
-		TimeRange = (Rate >= 0.0f) ? TRange<FMediaTimeStamp>(EstAudioTimeAtFrameStart, EstAudioTimeAtFrameStart + DeltaTime * Rate)
-			: TRange<FMediaTimeStamp>(EstAudioTimeAtFrameStart + DeltaTime * (1.0f + Rate), EstAudioTimeAtFrameStart + DeltaTime);
+		TimeRange = TRange<FMediaTimeStamp>(EstAudioTimeAtFrameStart, EstAudioTimeAtFrameStart + DeltaTime * FGenericPlatformMath::Abs(Rate));
 	}
 	else
 	{
@@ -2398,28 +2616,7 @@ bool FMediaPlayerFacade::GetCurrentPlaybackTimeRange(TRange<FMediaTimeStamp>& Ti
 					}
 				}
 
-				TimeRange = (Rate >= 0.0f) ? TRange<FMediaTimeStamp>(NextEstVideoTimeAtFrameStart.TimeStamp, NextEstVideoTimeAtFrameStart.TimeStamp + DeltaTime * Rate)
-					: TRange<FMediaTimeStamp>(NextEstVideoTimeAtFrameStart.TimeStamp + DeltaTime * (1.0f + Rate), NextEstVideoTimeAtFrameStart.TimeStamp + DeltaTime);
-
-				// If we are looping we check to prepare proper ranges should we wrap around either end of the media...
-				// (we do not clamp in the non-looping case as the rest of the code should deal with that fine)
-				if (Player->GetControls().IsLooping())
-				{
-					const FTimespan Duration = Player->GetControls().GetDuration();
-					FTimespan WrappedStart = WrappedModulo(TimeRange.GetLowerBoundValue().Time, Duration);
-					FTimespan WrappedEnd = WrappedModulo(TimeRange.GetUpperBoundValue().Time, Duration);
-					if (WrappedStart > WrappedEnd)
-					{
-						if (Rate >= 0.0)
-						{
-							TimeRange.SetUpperBoundValue(FMediaTimeStamp(WrappedEnd, TimeRange.GetUpperBoundValue().SequenceIndex + 1));
-						}
-						else
-						{
-							TimeRange.SetLowerBoundValue(FMediaTimeStamp(WrappedStart, TimeRange.GetLowerBoundValue().SequenceIndex - 1));
-						}
-					}
-				}
+				TimeRange = TRange<FMediaTimeStamp>(NextEstVideoTimeAtFrameStart.TimeStamp, NextEstVideoTimeAtFrameStart.TimeStamp + DeltaTime * FGenericPlatformMath::Abs(Rate));
 			}
 		}
 		else
@@ -2432,6 +2629,26 @@ bool FMediaPlayerFacade::GetCurrentPlaybackTimeRange(TRange<FMediaTimeStamp>& Ti
 		}
 	}
 
+	// If we are looping we check to prepare proper ranges should we wrap around either end of the media...
+	// (we do not clamp in the non-looping case as the rest of the code should deal with that fine)
+	if (Player->GetControls().IsLooping())
+	{
+		const FTimespan Duration = Player->GetControls().GetDuration();
+		FTimespan WrappedStart = WrappedModulo(TimeRange.GetLowerBoundValue().Time, Duration);
+		FTimespan WrappedEnd = WrappedModulo(TimeRange.GetUpperBoundValue().Time, Duration);
+		if (WrappedStart > WrappedEnd)
+		{
+			if (WrappedStart != TimeRange.GetLowerBoundValue().Time)
+			{
+				TimeRange.SetLowerBoundValue(FMediaTimeStamp(WrappedStart, FMediaTimeStamp::AdjustSecondaryIndex(TimeRange.GetLowerBoundValue().SequenceIndex, -1)));
+			}
+			if (WrappedEnd != TimeRange.GetUpperBoundValue().Time)
+			{
+				TimeRange.SetUpperBoundValue(FMediaTimeStamp(WrappedEnd, FMediaTimeStamp::AdjustSecondaryIndex(TimeRange.GetUpperBoundValue().SequenceIndex, 1)));
+			}
+		}
+	}
+
 	return !TimeRange.IsEmpty();
 }
 
@@ -2439,193 +2656,138 @@ bool FMediaPlayerFacade::GetCurrentPlaybackTimeRange(TRange<FMediaTimeStamp>& Ti
 /* FMediaPlayerFacade implementation
 *****************************************************************************/
 
-void FMediaPlayerFacade::ProcessAudioSamples(IMediaSamples& Samples, TRange<FTimespan> TimeRange)
+void FMediaPlayerFacade::ProcessAudioSamples(IMediaSamples& Samples, const TRange<FMediaTimeStamp>& TimeRange)
 {
 	TSharedPtr<IMediaAudioSample, ESPMode::ThreadSafe> Sample;
 
-	if (Player->GetPlayerFeatureFlag(IMediaPlayer::EFeatureFlag::UsePlaybackTimingV2))
+	check(Player->GetPlayerFeatureFlag(IMediaPlayer::EFeatureFlag::UsePlaybackTimingV2));
+
+	// For V2 we basically expect to get no timerange at all: totally open
+	// (we just have it around to be compatible / use older code that expects it)
+	check(TimeRange.GetLowerBound().IsOpen() && TimeRange.GetUpperBound().IsOpen());
+
+	// Seek in progress?
+	if (SeekTargetTime.IsValid())
 	{
-		// For V2 we basically expect to get no timerange at all: totally open
-		// (we just have it around to be compatible / use older code that expects it)
-		check(TimeRange.GetLowerBound().IsOpen() && TimeRange.GetUpperBound().IsOpen());
+		// Yes. Fetch (and discard) all samples up to the seek target time...
+		// (we only throw out samples from prior sequence indices to make sure we do not swallow any audio from overlapping samples)
+		auto DiscardRange = TRange<FMediaTimeStamp>(FMediaTimeStamp(0), FMediaTimeStamp((CurrentRate >= 0.0f) ? FTimespan::Zero() : FTimespan::MaxValue(), SeekTargetTime.SequenceIndex));
+		Samples.DiscardAudioSamples(DiscardRange, GetUnpausedRate() < 0.0f);
+	}
 
-		//
-		// "Modern" 1-Audio-Sink-Only case (aka: we only feed the primary sink)
-		//
-		if (TSharedPtr< FMediaAudioSampleSink, ESPMode::ThreadSafe> PinnedPrimaryAudioSink = PrimaryAudioSink.Pin())
+	//
+	// "Modern" 1-Audio-Sink-Only case (aka: we only feed the primary sink)
+	//
+	if (TSharedPtr< FMediaAudioSampleSink, ESPMode::ThreadSafe> PinnedPrimaryAudioSink = PrimaryAudioSink.Pin())
+	{
+		while (PinnedPrimaryAudioSink->CanAcceptSamples(1))
 		{
-			while (PinnedPrimaryAudioSink->CanAcceptSamples(1))
+			if (!Samples.FetchAudio(TimeRange, Sample))
+				break;
+
+			if (!Sample.IsValid())
 			{
-				if (!Samples.FetchAudio(TimeRange, Sample))
-					break;
+				continue;
+			}
+
+			{
+				FScopeLock Lock(&LastTimeValuesCS);
+				LastAudioSampleProcessedTime.TimeStamp = FMediaTimeStamp(Sample->GetTime());
+				LastAudioSampleProcessedTime.SampledAtTime = FPlatformTime::Seconds();
+			}
 
-				if (!Sample.IsValid())
+			PinnedPrimaryAudioSink->Enqueue(Sample.ToSharedRef());
+		}
+	}
+	else
+	{
+		// Do we have video playback?
+		if (HaveVideoPlayback())
+		{
+			TRange<FMediaTimeStamp> TempRange;
+			// We got video and audio, but no audio sink - throw away anything up to video playback time...
+			// (rough estimate, as this is off-gamethread; but better than throwing things out with no throttling at all)
+			{
+				bool bReverse = (CurrentRate < 0.0f);
+				FScopeLock Lock(&LastTimeValuesCS);
+				if (!bReverse)
 				{
-					continue;
+					TempRange.SetUpperBound(CurrentFrameVideoTimeStamp);
 				}
-
+				else
 				{
-					FScopeLock Lock(&LastTimeValuesCS);
-					LastAudioSampleProcessedTime.TimeStamp = FMediaTimeStamp(Sample->GetTime());
-					LastAudioSampleProcessedTime.SampledAtTime = FPlatformTime::Seconds();
+					TempRange.SetLowerBound(CurrentFrameVideoTimeStamp);
 				}
 
-				PinnedPrimaryAudioSink->Enqueue(Sample.ToSharedRef());
 			}
+			while (Samples.FetchAudio(TempRange, Sample))
+				;
 		}
 		else
 		{
-			// Do we have video playback?
-			if (HaveVideoPlayback())
-			{
-				TRange<FMediaTimeStamp> TempRange;
-				// We got video and audio, but no audio sink - throw away anything up to video playback time...
-				// (rough estimate, as this is off-gamethread; but better than throwing things out with no throttling at all)
-				{
-					bool bReverse = (CurrentRate < 0.0f);
-					FScopeLock Lock(&LastTimeValuesCS);
-					if (!bReverse)
-					{
-						TempRange.SetUpperBound(CurrentFrameVideoTimeStamp);
-					}
-					else
-					{
-						TempRange.SetLowerBound(CurrentFrameVideoTimeStamp);
-					}
+			// No Video and no primary audio sink: we throw all away (sub-optimal as it will keep audio decoding busy; but this should be an edge case)
+			while (Samples.FetchAudio(TimeRange, Sample))
+				;
+		}
+	}
+}
 
-				}
-				while (Samples.FetchAudio(TempRange, Sample))
-					;
-			}
-			else
-			{
-				// No Video and no primary audio sink: we throw all away (sub-optimal as it will keep audio decoding busy; but this should be an edge case)
-				while (Samples.FetchAudio(TimeRange, Sample))
-					;
-			}
-		}
-	}
-	else
-	{
-		//
-		// >1 Audio Sinks: we must drop samples that cause on overrun as SOME sinks will get it, some don't...
-		// (mainly here to cover the "backwards compatibility" cases -> in the future we will probably only allow ONE AudioSink)
-		//
-		while (Samples.FetchAudio(TimeRange, Sample))
-		{
-			if (!Sample.IsValid())
-			{
-				continue;
-			}
-
-			if (!AudioSampleSinks.Enqueue(Sample.ToSharedRef(), FMediaPlayerQueueDepths::MaxAudioSinkDepth))
-			{
-#if MEDIAPLAYERFACADE_TRACE_SINKOVERFLOWS
-				UE_LOG(LogMediaUtils, VeryVerbose, TEXT("PlayerFacade %p: Audio sample sink overflow"), this);
-#endif
-			}
-			else
-			{
-				FScopeLock Lock(&LastTimeValuesCS);
-				LastAudioSampleProcessedTime.TimeStamp = FMediaTimeStamp(Sample->GetTime());
-				LastAudioSampleProcessedTime.SampledAtTime = FPlatformTime::Seconds();
-			}
-		}
-	}
-}
-
-
-void FMediaPlayerFacade::ProcessCaptionSamples(IMediaSamples& Samples, TRange<FTimespan> TimeRange)
-{
-	TSharedPtr<IMediaOverlaySample, ESPMode::ThreadSafe> Sample;
-
-	while (Samples.FetchCaption(TimeRange, Sample))
-	{
-		if (Sample.IsValid() && !CaptionSampleSinks.Enqueue(Sample.ToSharedRef(), FMediaPlayerQueueDepths::MaxCaptionSinkDepth))
-		{
-#if MEDIAPLAYERFACADE_TRACE_SINKOVERFLOWS
-			UE_LOG(LogMediaUtils, VeryVerbose, TEXT("PlayerFacade %p: Caption sample sink overflow"), this);
-#endif
-		}
-	}
-}
 
-
-void FMediaPlayerFacade::ProcessMetadataSamples(IMediaSamples& Samples, TRange<FTimespan> TimeRange)
+bool FMediaPlayerFacade::IsVideoSampleStillGood(const TRange<FMediaTimeStamp>& LastSampleTimeRange, const TRange<FMediaTimeStamp>& TimeRange, bool bReverse) const
 {
-	TSharedPtr<IMediaBinarySample, ESPMode::ThreadSafe> Sample;
-
-	while (Samples.FetchMetadata(TimeRange, Sample))
+	// If we have no valid time range or a seek is in progress we assume the current frame can be considered "done" in any case
+	if (!TimeRange.IsEmpty() && !SeekTargetTime.IsValid())
 	{
-		if (Sample.IsValid() && !MetadataSampleSinks.Enqueue(Sample.ToSharedRef(), FMediaPlayerQueueDepths::MaxMetadataSinkDepth))
-		{
-#if MEDIAPLAYERFACADE_TRACE_SINKOVERFLOWS
-			UE_LOG(LogMediaUtils, VeryVerbose, TEXT("PlayerFacade %p: Metadata sample sink overflow"), this);
-#endif
-		}
-	}
-}
+		// This is not the case: check more detailed!
+		TRange<FMediaTimeStamp> LastSampleTimeRange0(FMediaTimeStamp(LastSampleTimeRange.GetLowerBoundValue().Time, 0), FMediaTimeStamp(LastSampleTimeRange.GetUpperBoundValue().Time, 0));
+		TRange<FMediaTimeStamp> TimeRange0(FMediaTimeStamp(TimeRange.GetLowerBoundValue().Time, 0), FMediaTimeStamp(TimeRange.GetUpperBoundValue().Time, 0));
 
-
-void FMediaPlayerFacade::ProcessSubtitleSamples(IMediaSamples& Samples, TRange<FTimespan> TimeRange)
-{
-	TSharedPtr<IMediaOverlaySample, ESPMode::ThreadSafe> Sample;
-
-	while (Samples.FetchSubtitle(TimeRange, Sample))
-	{
-		//UE_LOG(LogMediaUtils, Display, TEXT("Subtitle @%.3f: %s"), Sample->GetTime().Time.GetTotalSeconds(), *Sample->GetText().ToString());
-		if (Sample.IsValid() && !SubtitleSampleSinks.Enqueue(Sample.ToSharedRef(), FMediaPlayerQueueDepths::MaxSubtitleSinkDepth))
+		// Is the sample time range at all still valid?
+		if (LastSampleTimeRange0.Overlaps(TimeRange0))
 		{
-#if MEDIAPLAYERFACADE_TRACE_SINKOVERFLOWS
-			UE_LOG(LogMediaUtils, VeryVerbose, TEXT("PlayerFacade %p: Subtitle sample sink overflow"), this);
-#endif
-		}
-	}
-}
-
+			// Yes. Assuming we could get more samples (of the same type) from the player, would the next one "better"?
+			// (this does not say which "next one" is actually current (e.g. if the rendering FPS is real bad), just that "another one" would be better)
 
-void FMediaPlayerFacade::ProcessVideoSamplesV1(IMediaSamples& Samples, TRange<FTimespan> TimeRange)
-{
-	// Let the player do some processing if needed.
-	if (Player.IsValid())
-	{
-		Player->ProcessVideoSamples();
-	}
-
-	// This is not to be used with V2 timing
-	check(!Player->GetPlayerFeatureFlag(IMediaPlayer::EFeatureFlag::UsePlaybackTimingV2));
-
-	TSharedPtr<IMediaTextureSample, ESPMode::ThreadSafe> Sample;
+			// Compute the "theoretical" next sample range...
+			TRange<FMediaTimeStamp> NextSampleTimeRange = !bReverse ? TRange<FMediaTimeStamp>(LastSampleTimeRange0.GetUpperBoundValue(), LastSampleTimeRange0.GetUpperBoundValue() + LastSampleTimeRange0.Size<FMediaTimeStamp>().Time)
+				: TRange<FMediaTimeStamp>(LastSampleTimeRange0.GetLowerBoundValue() - LastSampleTimeRange0.Size<FMediaTimeStamp>().Time, LastSampleTimeRange0.GetLowerBoundValue());
 
-	while (Samples.FetchVideo(TimeRange, Sample))
-	{
-		if (!Sample.IsValid())
-		{
-			continue;
-		}
+			FTimespan Duration = Player->GetControls().GetDuration();
 
-		{
-			FScopeLock Lock(&LastTimeValuesCS);
-			CurrentFrameVideoTimeStamp = Sample->GetTime();
-		}
+			if (!Player->GetControls().IsLooping())
+			{
+				// If we are not looping we need to clamp against the media's duration
+				// (we assume it starts at zero here!)
+				NextSampleTimeRange = TRange<FMediaTimeStamp>::Intersection(NextSampleTimeRange, TRange<FMediaTimeStamp>(FMediaTimeStamp(0, NextSampleTimeRange.GetLowerBoundValue().SequenceIndex), FMediaTimeStamp(Duration, NextSampleTimeRange.GetLowerBoundValue().SequenceIndex)));
+			}
+			else
+			{
+				if (NextSampleTimeRange.GetLowerBoundValue().Time >= Duration)
+				{
+					check(!bReverse);
+					NextSampleTimeRange = TRange<FMediaTimeStamp>(FMediaTimeStamp(NextSampleTimeRange.GetLowerBoundValue().Time - Duration, FMediaTimeStamp::AdjustPrimaryIndex(NextSampleTimeRange.GetLowerBoundValue().SequenceIndex, 1)),
+						FMediaTimeStamp(NextSampleTimeRange.GetUpperBoundValue().Time - Duration, FMediaTimeStamp::AdjustPrimaryIndex(NextSampleTimeRange.GetUpperBoundValue().SequenceIndex, 1)));
+				}
+				else if (NextSampleTimeRange.GetLowerBoundValue().Time < FTimespan::Zero())
+				{
+					check(bReverse);
+					NextSampleTimeRange = TRange<FMediaTimeStamp>(FMediaTimeStamp(NextSampleTimeRange.GetLowerBoundValue().Time + Duration, FMediaTimeStamp::AdjustPrimaryIndex(NextSampleTimeRange.GetLowerBoundValue().SequenceIndex, -1)),
+						FMediaTimeStamp(NextSampleTimeRange.GetUpperBoundValue().Time + Duration, FMediaTimeStamp::AdjustPrimaryIndex(NextSampleTimeRange.GetUpperBoundValue().SequenceIndex, -1)));
+				}
+			}
 
-		UE_LOG(LogMediaUtils, VeryVerbose, TEXT("PlayerFacade %p: Fetched video sample %s"), this, *Sample->GetTime().Time.ToString(TEXT("%h:%m:%s.%t")));
+			// Compute which one is larger inside the current range...
+			int64 LastSampleCoverage = TRange<FMediaTimeStamp>::Intersection(TimeRange0, LastSampleTimeRange0).Size<FMediaTimeStamp>().Time.GetTicks();
+			int64 NextSampleCoverage = TRange<FMediaTimeStamp>::Intersection(TimeRange0, NextSampleTimeRange).Size<FMediaTimeStamp>().Time.GetTicks();
 
-		if (VideoSampleSinks.Enqueue(Sample.ToSharedRef(), FMediaPlayerQueueDepths::MaxVideoSinkDepth))
-		{
-			if (CurrentRate >= 0.0f)
+			if (LastSampleCoverage > NextSampleCoverage)
 			{
-				NextVideoSampleTime = Sample->GetTime().Time + Sample->GetDuration();
-				UE_LOG(LogMediaUtils, VeryVerbose, TEXT("PlayerFacade %p: Next video sample time %s"), this, *NextVideoSampleTime.ToString(TEXT("%h:%m:%s.%t")));
+				// Last one we returned is still good. No new one needed...
+				return true;
 			}
 		}
-		else
-		{
-#if MEDIAPLAYERFACADE_TRACE_SINKOVERFLOWS
-			UE_LOG(LogMediaUtils, VeryVerbose, TEXT("PlayerFacade %p: Video sample sink overflow"), this);
-#endif
-		}
 	}
+	return false;
 }
 
 
@@ -2652,6 +2814,12 @@ bool FMediaPlayerFacade::ProcessVideoSamples(IMediaSamples& Samples, const TRang
 		//
 		const bool bReverse = (GetUnpausedRate() < 0.0f);
 
+		if (IsVideoSampleStillGood(LastVideoSampleProcessedTimeRange, TimeRange, bReverse))
+		{
+			// We got all the samples we need. Processing was successful...
+			return true;
+		}
+
 		switch (Samples.FetchBestVideoSampleForTimeRange(TimeRange, Sample, bReverse))
 		{
 		case IMediaSamples::EFetchBestSampleResult::Ok:
@@ -2717,9 +2885,13 @@ bool FMediaPlayerFacade::ProcessVideoSamples(IMediaSamples& Samples, const TRang
 			bool bOk = VideoSampleSinks.Enqueue(Sample.ToSharedRef(), FMediaPlayerQueueDepths::MaxVideoSinkDepth);
 			check(bOk);
 
-			FScopeLock Lock(&LastTimeValuesCS);
-			CurrentFrameVideoTimeStamp = SampleTimeRange.GetLowerBoundValue();
-			LastVideoSampleProcessedTimeRange = SampleTimeRange;
+			{
+				FScopeLock Lock(&LastTimeValuesCS);
+				CurrentFrameVideoDisplayTimeStamp = CurrentFrameVideoTimeStamp = SampleTimeRange.GetLowerBoundValue();
+				LastVideoSampleProcessedTimeRange = SampleTimeRange;
+			}
+
+			UpdateSeekStatus(&CurrentFrameVideoTimeStamp);
 
 			return true;
 		}
@@ -2728,40 +2900,184 @@ bool FMediaPlayerFacade::ProcessVideoSamples(IMediaSamples& Samples, const TRang
 }
 
 
-void FMediaPlayerFacade::ProcessCaptionSamples(IMediaSamples& Samples, TRange<FMediaTimeStamp> TimeRange)
+void FMediaPlayerFacade::ProcessCaptionSamples(IMediaSamples& Samples, const TRange<FMediaTimeStamp>& TimeRange)
 {
+	check(Player->GetPlayerFeatureFlag(IMediaPlayer::EFeatureFlag::UsePlaybackTimingV2));
+
 	TSharedPtr<IMediaOverlaySample, ESPMode::ThreadSafe> Sample;
 
-	while (Samples.FetchCaption(TimeRange, Sample))
+	// Seek in progress?
+	if (SeekTargetTime.IsValid())
 	{
-		if (Sample.IsValid() && !CaptionSampleSinks.Enqueue(Sample.ToSharedRef(), FMediaPlayerQueueDepths::MaxCaptionSinkDepth))
+		// Yes. Fetch (and discard) all samples up to the seek target time...
+		// (we only throw out samples from prior sequence indices to make sure we do not swallow any audio from overlapping samples)
+		auto DiscardRange = TRange<FMediaTimeStamp>(FMediaTimeStamp(0), FMediaTimeStamp((CurrentRate >= 0.0f) ? FTimespan::Zero() : FTimespan::MaxValue(), SeekTargetTime.SequenceIndex));
+		Samples.DiscardCaptionSamples(DiscardRange, GetUnpausedRate() < 0.0f);
+	}
+	else
+	{
+		while (Samples.FetchCaption(TimeRange, Sample))
 		{
+			if (Sample.IsValid() && !CaptionSampleSinks.Enqueue(Sample.ToSharedRef(), FMediaPlayerQueueDepths::MaxCaptionSinkDepth))
+			{
 #if MEDIAPLAYERFACADE_TRACE_SINKOVERFLOWS
-			UE_LOG(LogMediaUtils, VeryVerbose, TEXT("PlayerFacade %p: Caption sample sink overflow"), this);
+				UE_LOG(LogMediaUtils, VeryVerbose, TEXT("PlayerFacade %p: Caption sample sink overflow"), this);
 #endif
+			}
 		}
 	}
 }
 
 
-void FMediaPlayerFacade::ProcessSubtitleSamples(IMediaSamples& Samples, TRange<FMediaTimeStamp> TimeRange)
+void FMediaPlayerFacade::ProcessSubtitleSamples(IMediaSamples& Samples, const TRange<FMediaTimeStamp>& TimeRange)
 {
+	check(Player->GetPlayerFeatureFlag(IMediaPlayer::EFeatureFlag::UsePlaybackTimingV2));
+
 	TSharedPtr<IMediaOverlaySample, ESPMode::ThreadSafe> Sample;
 
-	// There might be samples in the subtitle sample queue that have may lie before the specified time range.
-	// Since these have no way of being displayed they must be removed from the sample queue.
-	TSharedPtr<IMediaPlayer, ESPMode::ThreadSafe> CurrentPlayer = Player;
-	if (CurrentPlayer.IsValid() && CurrentPlayer->GetPlayerFeatureFlag(IMediaPlayer::EFeatureFlag::UsePlaybackTimingV2))
+	// Seek in progress?
+	if (SeekTargetTime.IsValid())
+	{
+		// Yes. Fetch (and discard) all samples up to the seek target time...
+		// (we only throw out samples from prior sequence indices to make sure we do not swallow any audio from overlapping samples)
+		auto DiscardRange = TRange<FMediaTimeStamp>(FMediaTimeStamp(0), FMediaTimeStamp((CurrentRate >= 0.0f) ? FTimespan::Zero() : FTimespan::MaxValue(), SeekTargetTime.SequenceIndex));
+		Samples.DiscardSubtitleSamples(DiscardRange, GetUnpausedRate() < 0.0f);
+	}
+	else
+	{
+		while (Samples.FetchSubtitle(TimeRange, Sample))
+		{
+			//UE_LOG(LogMediaUtils, Display, TEXT("Subtitle @%.3f: %s"), Sample->GetTime().Time.GetTotalSeconds(), *Sample->GetText().ToString());
+			if (Sample.IsValid() && !SubtitleSampleSinks.Enqueue(Sample.ToSharedRef(), FMediaPlayerQueueDepths::MaxSubtitleSinkDepth))
+			{
+#if MEDIAPLAYERFACADE_TRACE_SINKOVERFLOWS
+				UE_LOG(LogMediaUtils, VeryVerbose, TEXT("PlayerFacade %p: Subtitle sample sink overflow"), this);
+#endif
+			}
+		}
+	}
+}
+
+
+void FMediaPlayerFacade::ProcessMetadataSamples(IMediaSamples& Samples, const TRange<FMediaTimeStamp>& TimeRange)
+{
+	check(Player->GetPlayerFeatureFlag(IMediaPlayer::EFeatureFlag::UsePlaybackTimingV2));
+
+	TSharedPtr<IMediaBinarySample, ESPMode::ThreadSafe> Sample;
+
+	// Seek in progress?
+	if (SeekTargetTime.IsValid())
+	{
+		// Yes. Fetch (and discard) all samples up to the seek target time...
+		// (we only throw out samples from prior sequence indices to make sure we do not swallow any audio from overlapping samples)
+		auto DiscardRange = TRange<FMediaTimeStamp>(FMediaTimeStamp(0), FMediaTimeStamp((CurrentRate >= 0.0f) ? FTimespan::Zero() : FTimespan::MaxValue(), SeekTargetTime.SequenceIndex));
+		Samples.DiscardMetadataSamples(DiscardRange, GetUnpausedRate() < 0.0f);
+	}
+	else
+	{
+		while (Samples.FetchMetadata(TimeRange, Sample))
+		{
+			if (Sample.IsValid() && !MetadataSampleSinks.Enqueue(Sample.ToSharedRef(), FMediaPlayerQueueDepths::MaxMetadataSinkDepth))
+			{
+#if MEDIAPLAYERFACADE_TRACE_SINKOVERFLOWS
+				UE_LOG(LogMediaUtils, VeryVerbose, TEXT("PlayerFacade %p: Metadata sample sink overflow"), this);
+#endif
+			}
+		}
+	}
+}
+
+
+void FMediaPlayerFacade::ProcessAudioSamplesV1(IMediaSamples& Samples, TRange<FTimespan> TimeRange)
+{
+	TSharedPtr<IMediaAudioSample, ESPMode::ThreadSafe> Sample;
+
+	while (Samples.FetchAudio(TimeRange, Sample))
 	{
-		uint32 NumPurged = 0;
-		if (CurrentPlayer->GetSamples().NumSubtitleSamples() >= kMinFramesInSubtitleQueueToPurge)
+		if (!Sample.IsValid())
+		{
+			continue;
+		}
+
+		if (!AudioSampleSinks.Enqueue(Sample.ToSharedRef(), FMediaPlayerQueueDepths::MaxAudioSinkDepth))
 		{
-			bool bReverse = CurrentRate < 0.0f;
-			NumPurged = CurrentPlayer->GetSamples().PurgeOutdatedSubtitleSamples(TimeRange.GetLowerBoundValue() + FTimespan::FromSeconds((bReverse ? kOutdatedSubtitleSamplesTolerance : -kOutdatedSubtitleSamplesTolerance)), bReverse);
+#if MEDIAPLAYERFACADE_TRACE_SINKOVERFLOWS
+			UE_LOG(LogMediaUtils, VeryVerbose, TEXT("PlayerFacade %p: Audio sample sink overflow"), this);
+#endif
+		}
+		else
+		{
+			FScopeLock Lock(&LastTimeValuesCS);
+			LastAudioSampleProcessedTime.TimeStamp = Sample->GetTime();
+			LastAudioSampleProcessedTime.SampledAtTime = FPlatformTime::Seconds();
 		}
-		SET_DWORD_STAT(STAT_MediaUtils_FacadeNumPurgedSubtitleSamples, NumPurged);
-		INC_DWORD_STAT_BY(STAT_MediaUtils_FacadeTotalPurgedSubtitleSamples, NumPurged);
 	}
+}
+
+
+void FMediaPlayerFacade::ProcessVideoSamplesV1(IMediaSamples& Samples, TRange<FTimespan> TimeRange)
+{
+	// Let the player do some processing if needed.
+	if (Player.IsValid())
+	{
+		Player->ProcessVideoSamples();
+	}
+
+	// This is not to be used with V2 timing
+	check(!Player->GetPlayerFeatureFlag(IMediaPlayer::EFeatureFlag::UsePlaybackTimingV2));
+
+	TSharedPtr<IMediaTextureSample, ESPMode::ThreadSafe> Sample;
+
+	while (Samples.FetchVideo(TimeRange, Sample))
+	{
+		if (!Sample.IsValid())
+		{
+			continue;
+		}
+
+		{
+			FScopeLock Lock(&LastTimeValuesCS);
+			CurrentFrameVideoDisplayTimeStamp = CurrentFrameVideoTimeStamp = Sample->GetTime();
+		}
+
+		UE_LOG(LogMediaUtils, VeryVerbose, TEXT("PlayerFacade %p: Fetched video sample %s"), this, *Sample->GetTime().Time.ToString(TEXT("%h:%m:%s.%t")));
+
+		if (VideoSampleSinks.Enqueue(Sample.ToSharedRef(), FMediaPlayerQueueDepths::MaxVideoSinkDepth))
+		{
+			if (CurrentRate >= 0.0f)
+			{
+				NextVideoSampleTime = Sample->GetTime().Time + Sample->GetDuration();
+				UE_LOG(LogMediaUtils, VeryVerbose, TEXT("PlayerFacade %p: Next video sample time %s"), this, *NextVideoSampleTime.ToString(TEXT("%h:%m:%s.%t")));
+			}
+		}
+		else
+		{
+#if MEDIAPLAYERFACADE_TRACE_SINKOVERFLOWS
+			UE_LOG(LogMediaUtils, VeryVerbose, TEXT("PlayerFacade %p: Video sample sink overflow"), this);
+#endif
+		}
+	}
+}
+
+void FMediaPlayerFacade::ProcessCaptionSamplesV1(IMediaSamples& Samples, TRange<FTimespan> TimeRange)
+{
+	TSharedPtr<IMediaOverlaySample, ESPMode::ThreadSafe> Sample;
+
+	while (Samples.FetchCaption(TimeRange, Sample))
+	{
+		if (Sample.IsValid() && !CaptionSampleSinks.Enqueue(Sample.ToSharedRef(), FMediaPlayerQueueDepths::MaxCaptionSinkDepth))
+		{
+#if MEDIAPLAYERFACADE_TRACE_SINKOVERFLOWS
+			UE_LOG(LogMediaUtils, VeryVerbose, TEXT("PlayerFacade %p: Caption sample sink overflow"), this);
+#endif
+		}
+	}
+}
+
+
+void FMediaPlayerFacade::ProcessSubtitleSamplesV1(IMediaSamples& Samples, TRange<FTimespan> TimeRange)
+{
+	TSharedPtr<IMediaOverlaySample, ESPMode::ThreadSafe> Sample;
 
 	while (Samples.FetchSubtitle(TimeRange, Sample))
 	{
@@ -2778,6 +3094,22 @@ void FMediaPlayerFacade::ProcessSubtitleSamples(IMediaSamples& Samples, TRange<F
 }
 
 
+void FMediaPlayerFacade::ProcessMetadataSamplesV1(IMediaSamples& Samples, TRange<FTimespan> TimeRange)
+{
+	TSharedPtr<IMediaBinarySample, ESPMode::ThreadSafe> Sample;
+
+	while (Samples.FetchMetadata(TimeRange, Sample))
+	{
+		if (Sample.IsValid() && !MetadataSampleSinks.Enqueue(Sample.ToSharedRef(), FMediaPlayerQueueDepths::MaxMetadataSinkDepth))
+		{
+#if MEDIAPLAYERFACADE_TRACE_SINKOVERFLOWS
+			UE_LOG(LogMediaUtils, VeryVerbose, TEXT("PlayerFacade %p: Metadata sample sink overflow"), this);
+#endif
+		}
+	}
+}
+
+
 /* IMediaEventSink interface
 *****************************************************************************/
 
@@ -2830,17 +3162,38 @@ void FMediaPlayerFacade::ReceiveMediaEvent(EMediaEvent Event)
 				return;
 			}
 
-			uint32 NumPurged = 0;
+			bool bReverse = (Rate < 0.0f);
+			const float RateFactor = (Rate != 0.0f) ? (1.0f / Rate) : 1.0f;
+
 			// Don't purge frames if the queue is small (to avoid purging if players deliver frames late persistently)
+			uint32 NumPurged = 0;
 			if (CurrentPlayer->GetSamples().NumVideoSamples() >= kMinFramesInVideoQueueToPurge)
 			{
-				bool bReverse = (Rate < 0.0f);
-				NumPurged = CurrentPlayer->GetSamples().PurgeOutdatedVideoSamples(TimeRange.GetLowerBoundValue() + FTimespan::FromSeconds(bReverse ? kOutdatedVideoSamplesTolerance : -kOutdatedVideoSamplesTolerance), bReverse);
+				NumPurged = CurrentPlayer->GetSamples().PurgeOutdatedVideoSamples(TimeRange.GetLowerBoundValue(), bReverse, FTimespan::FromSeconds(kOutdatedVideoSamplesTolerance * RateFactor));
 			}
-
 			SET_DWORD_STAT(STAT_MediaUtils_FacadeNumPurgedVideoSamples, NumPurged);
 			INC_DWORD_STAT_BY(STAT_MediaUtils_FacadeTotalPurgedVideoSamples, NumPurged);
 
+			// Take the opportunity to also purge any samples related to video samples directly (and evaluated on the game thread)
+
+			// Captions...
+			NumPurged = 0;
+			if (CurrentPlayer->GetSamples().NumVideoSamples() >= kMinFramesInCaptionQueueToPurge)
+			{
+				NumPurged = CurrentPlayer->GetSamples().PurgeOutdatedCaptionSamples(TimeRange.GetLowerBoundValue(), bReverse, FTimespan::FromSeconds(kOutdatedVideoSamplesTolerance * RateFactor));
+			}
+			SET_DWORD_STAT(STAT_MediaUtils_FacadeNumPurgedSubtitleSamples, NumPurged);
+			INC_DWORD_STAT_BY(STAT_MediaUtils_FacadeTotalPurgedSubtitleSamples, NumPurged);
+
+			// Subtitles...
+			NumPurged = 0;
+			if (CurrentPlayer->GetSamples().NumVideoSamples() >= kMinFramesInSubtitleQueueToPurge)
+			{
+				NumPurged = CurrentPlayer->GetSamples().PurgeOutdatedSubtitleSamples(TimeRange.GetLowerBoundValue(), bReverse, FTimespan::FromSeconds(kOutdatedVideoSamplesTolerance * RateFactor));
+			}
+			SET_DWORD_STAT(STAT_MediaUtils_FacadeNumPurgedCaptionSamples, NumPurged);
+			INC_DWORD_STAT_BY(STAT_MediaUtils_FacadeTotalPurgedCaptionSamples, NumPurged);
+
 			break;
 		}
 
diff --git a/Engine/Source/Runtime/MediaUtils/Private/MediaSamples.cpp b/Engine/Source/Runtime/MediaUtils/Private/MediaSamples.cpp
--- a/Engine/Source/Runtime/MediaUtils/Private/MediaSamples.cpp
+++ b/Engine/Source/Runtime/MediaUtils/Private/MediaSamples.cpp
@@ -107,6 +107,11 @@ bool FMediaSamples::FetchSubtitle(TRange<FMediaTimeStamp> TimeRange, TSharedPtr<
 	return FetchSample(SubtitleSampleQueue, TimeRange, OutSample);
 }
 
+bool FMediaSamples::FetchMetadata(TRange<FMediaTimeStamp> TimeRange, TSharedPtr<IMediaBinarySample, ESPMode::ThreadSafe>& OutSample)
+{
+	return FetchSample(MetadataSampleQueue, TimeRange, OutSample);
+}
+
 bool FMediaSamples::FetchVideo(TRange<FMediaTimeStamp> TimeRange, TSharedPtr<IMediaTextureSample, ESPMode::ThreadSafe>& OutSample)
 {
 	return FetchSample(VideoSampleQueue, TimeRange, OutSample);
@@ -139,17 +144,33 @@ FMediaSamples::EFetchBestSampleResult FMediaSamples::FetchBestVideoSampleForTime
 /**
  * Remove any video samples from the queue that have no chance of being displayed anymore
  */
-uint32 FMediaSamples::PurgeOutdatedVideoSamples(const FMediaTimeStamp & ReferenceTime, bool bReversed)
+uint32 FMediaSamples::PurgeOutdatedVideoSamples(const FMediaTimeStamp & ReferenceTime, bool bReversed, FTimespan MaxAge)
 {
-	return VideoSampleQueue.PurgeOutdatedSamples(ReferenceTime, bReversed);
+	return VideoSampleQueue.PurgeOutdatedSamples(ReferenceTime, bReversed, MaxAge);
 }
 
 /**
  * Remove any subtitle samples from the queue that have no chance of being displayed anymore
  */
-uint32 FMediaSamples::PurgeOutdatedSubtitleSamples(const FMediaTimeStamp & ReferenceTime, bool bReversed)
+uint32 FMediaSamples::PurgeOutdatedSubtitleSamples(const FMediaTimeStamp & ReferenceTime, bool bReversed, FTimespan MaxAge)
+{
+	return SubtitleSampleQueue.PurgeOutdatedSamples(ReferenceTime, bReversed, MaxAge);
+}
+
+/**
+ * Remove any caption samples from the queue that have no chance of being displayed anymore
+ */
+uint32 FMediaSamples::PurgeOutdatedCaptionSamples(const FMediaTimeStamp& ReferenceTime, bool bReversed, FTimespan MaxAge)
+{
+	return CaptionSampleQueue.PurgeOutdatedSamples(ReferenceTime, bReversed, MaxAge);
+}
+
+/**
+ * Remove any caption samples from the queue that have no chance of being displayed anymore
+ */
+uint32 FMediaSamples::PurgeOutdatedMetadataSamples(const FMediaTimeStamp& ReferenceTime, bool bReversed, FTimespan MaxAge)
 {
-	return SubtitleSampleQueue.PurgeOutdatedSamples(ReferenceTime, bReversed);
+	return MetadataSampleQueue.PurgeOutdatedSamples(ReferenceTime, bReversed, MaxAge);
 }
 
 /**
diff --git a/Engine/Source/Runtime/MediaUtils/Public/MediaPlayerFacade.h b/Engine/Source/Runtime/MediaUtils/Public/MediaPlayerFacade.h
--- a/Engine/Source/Runtime/MediaUtils/Public/MediaPlayerFacade.h
+++ b/Engine/Source/Runtime/MediaUtils/Public/MediaPlayerFacade.h
@@ -334,10 +334,19 @@ public:
 	 * Get the media's current playback time stamp.
 	 *
 	 * @return Playback time stamp.
-	 * @see GetDuration, Seek
 	 */
 	FMediaTimeStamp GetTimeStamp() const;
 
+	/**
+	 * Get the media's current playback time stamp in a "display" version
+	 *
+	 * @return Playback time stamp.
+	 * 
+	 * @note The timestamp returned here will reflect a user-logic oriented version.
+	 *       (e.g. during seeks this will return the seek target rather than the last valid frame still displayed)
+	 */
+	FMediaTimeStamp GetDisplayTimeStamp() const;
+
 	/**
 	 * Get the human readable name of the specified track.
 	 *
@@ -742,7 +751,10 @@ protected:
 	bool BlockOnFetch() const;
 
 	/** Flush all media sample sinks & player plugin. */
-	void Flush(bool bExcludePlayer = false);
+	void Flush(bool bExcludePlayer = false, bool bOnSeek = false);
+
+	/** Internal function to retrieve the current timestamp */
+	FMediaTimeStamp GetTimeStampInternal(bool bForDisplay) const;
 
 	/**
 	 * Get details about the specified audio track format.
@@ -791,20 +803,6 @@ protected:
 	bool HaveVideoPlayback() const;
 	float GetUnpausedRate() const;
 
-	/** Fetch audio samples from the player and forward them to the registered sinks. */
-	void ProcessAudioSamples(IMediaSamples& Samples, TRange<FTimespan> TimeRange);
-
-	/** Fetch metadata samples from the player and forward them to the registered sinks. */
-	void ProcessMetadataSamples(IMediaSamples& Samples, TRange<FTimespan> TimeRange);
-
-	/** Fetch audio samples from the player and forward them to the registered sinks. */
-	void ProcessCaptionSamples(IMediaSamples& Samples, TRange<FTimespan> TimeRange);
-
-	/** Fetch subtitle samples from the player and forward them to the registered sinks. */
-	void ProcessSubtitleSamples(IMediaSamples& Samples, TRange<FTimespan> TimeRange);
-
-	/** Fetch video samples from the player and forward them to the registered sinks. */
-	void ProcessVideoSamplesV1(IMediaSamples& Samples, TRange<FTimespan> TimeRange);
 
 protected:
 
@@ -823,12 +821,23 @@ private:
 	bool NotifyLifetimeManagerDelegate_PlayerDestroyed();
 	bool NotifyLifetimeManagerDelegate_PlayerResourcesReleased(uint32 ResourceFlags);
 
+	void ProcessAudioSamples(IMediaSamples& Samples, const TRange<FMediaTimeStamp>& TimeRange);
 	bool ProcessVideoSamples(IMediaSamples& Samples, const TRange<FMediaTimeStamp>& TimeRange);
-	void ProcessCaptionSamples(IMediaSamples& Samples, TRange<FMediaTimeStamp> TimeRange);
-	void ProcessSubtitleSamples(IMediaSamples& Samples, TRange<FMediaTimeStamp> TimeRange);
+	void ProcessCaptionSamples(IMediaSamples& Samples, const TRange<FMediaTimeStamp>& TimeRange);
+	void ProcessSubtitleSamples(IMediaSamples& Samples, const TRange<FMediaTimeStamp>& TimeRange);
+	void ProcessMetadataSamples(IMediaSamples& Samples, const TRange<FMediaTimeStamp>& TimeRange);
 
+	void ProcessAudioSamplesV1(IMediaSamples& Samples, TRange<FTimespan> TimeRange);
+	void ProcessVideoSamplesV1(IMediaSamples& Samples, TRange<FTimespan> TimeRange);
+	void ProcessSubtitleSamplesV1(IMediaSamples& Samples, TRange<FTimespan> TimeRange);
+	void ProcessCaptionSamplesV1(IMediaSamples& Samples, TRange<FTimespan> TimeRange);
+	void ProcessMetadataSamplesV1(IMediaSamples& Samples, TRange<FTimespan> TimeRange);
+
+	bool IsVideoSampleStillGood(const TRange<FMediaTimeStamp>& LastSampleTimeRange, const TRange<FMediaTimeStamp>& TimeRange, bool bReverse) const;
+	void MonitorAudioEnablement();
+	void UpdateSeekStatus(const FMediaTimeStamp* pCheckTimeStamp = nullptr);
 	void PreSampleProcessingTimeHandling();
-	bool GetCurrentPlaybackTimeRange(TRange<FMediaTimeStamp>& TimeRange, float Rate, FTimespan DeltaTime, bool bDoNotUseFrameStartReference) const;
+	bool GetCurrentPlaybackTimeRange(TRange<FMediaTimeStamp>& TimeRange, float Rate, FTimespan DeltaTime, bool bPurgeSampleRelated) const;
 	void PostSampleProcessingTimeHandling(FTimespan DeltaTime);
 
 	void DestroyPlayer();
@@ -863,7 +872,8 @@ private:
 		const TRange<FMediaTimeStamp> & GetRange() const;
 		bool IsSet() const;
 
-		void Flush();
+		void OnFlush();
+		void OnSeek(int32 PrimaryIndex);
 
 		void Reset()
 		{
@@ -871,7 +881,8 @@ private:
 			CurrentTimeRange = TRange<FTimespan>::Empty();
 			LastBlockOnRange = TRange<FTimespan>::Empty();
 			RangeIsDirty = false;
-			OnBlockSeqIndex = 0;
+			OnBlockPrimaryIndex = 0;
+			OnBlockSecondaryIndex = 0;
 		}
 
 	private:
@@ -890,8 +901,10 @@ private:
 		/** Flag to indicate if internal range is valid or not */
 		mutable bool RangeIsDirty;
 
-		/** Sequence index used during blocked playback processing */
-		mutable int64 OnBlockSeqIndex;
+		/** Primary sequence index used during blocked playback processing */
+		mutable int32 OnBlockPrimaryIndex;
+		/** Secondary sequence index used during blocked playback processing */
+		mutable int32 OnBlockSecondaryIndex;
 	};
 
 	FBlockOnRange BlockOnRange;
@@ -964,9 +977,18 @@ private:
 	/** Timestamp for video considered "current" for this frame (stays valid even after a flush, until new data comes in) */
 	FMediaTimeStamp CurrentFrameVideoTimeStamp;
 
+	/** Timestamp for video considered "current" for this frame for GUI / display purposes (stays valid even after a flush, until new data comes in) */
+	FMediaTimeStamp CurrentFrameVideoDisplayTimeStamp;
+
 	/** Estimation for next frame's video timestamp (used when no audio present or active in stream) */
 	FMediaTimeStampSample NextEstVideoTimeAtFrameStart;
 
+	/** Timestamp of seek target location if seek is pending */
+	FMediaTimeStamp SeekTargetTime;
+
+	/** Current seek index */
+	int32 SeekIndex;
+
 	/** Set if sinks are to be flushed at the request of the player. */
 	TAtomic<bool>	bIsSinkFlushPending;
 
diff --git a/Engine/Source/Runtime/MediaUtils/Public/MediaSampleQueue.h b/Engine/Source/Runtime/MediaUtils/Public/MediaSampleQueue.h
--- a/Engine/Source/Runtime/MediaUtils/Public/MediaSampleQueue.h
+++ b/Engine/Source/Runtime/MediaUtils/Public/MediaSampleQueue.h
@@ -16,7 +16,10 @@
 #include "MediaSampleSource.h"
 
 #include "IMediaTimeSource.h"
+#include "IMediaAudioSample.h"
 #include "IMediaTextureSample.h"
+#include "IMediaBinarySample.h"
+#include "IMediaOverlaySample.h"
 
 /**
  * Template for media sample queues.
@@ -119,63 +122,39 @@ public:
 		return true;
 	}
 
-	bool FetchBestSampleForTimeRange(const TRange<FMediaTimeStamp> & TimeRange, TSharedPtr<SampleType, ESPMode::ThreadSafe>& OutSample, bool bReverse)
+	bool Discard(const TRange<FMediaTimeStamp>& TimeRange, bool bReverse)
 	{
 		// Code below assumes a fully specified range, no open bounds!
 		check(TimeRange.HasLowerBound() && TimeRange.HasUpperBound());
 
-		OutSample.Reset();
-
 		FScopeLock Lock(&CriticalSection);
 
-		int32 Num = Samples.Num();
-		if (Num == 0)
+		int32 FirstPossibleIndex, LastPossibleIndex, NumOldSamplesAtBegin;
+		FindRangeInQueue(TimeRange, bReverse, FirstPossibleIndex, LastPossibleIndex, NumOldSamplesAtBegin);
+
+		// Found anything?
+		if (FirstPossibleIndex >= 0)
 		{
-			return false;
+			check(LastPossibleIndex >= 0);
+
+			// Remove all samples indicated...
+			Samples.RemoveAt(FirstPossibleIndex, LastPossibleIndex - FirstPossibleIndex + 1);
+			return true;
 		}
+		return false;
+	}
 
-		int32 FirstPossibleIndex = -1;
-		int32 LastPossibleIndex = -1;
-		int32 NumOldSamplesAtBegin = 0;
-		for (int32 Idx = 0; Idx < Num; ++Idx)
-		{
-			const TSharedPtr<SampleType, ESPMode::ThreadSafe> & Sample = Samples[Idx];
-			TRange<FMediaTimeStamp> SampleTimeRange = !bReverse ? TRange<FMediaTimeStamp>(Sample->GetTime(), Sample->GetTime() + Sample->GetDuration())
-																: TRange<FMediaTimeStamp>(Sample->GetTime() - Sample->GetDuration(), Sample->GetTime());
+	bool FetchBestSampleForTimeRange(const TRange<FMediaTimeStamp>& TimeRange, TSharedPtr<SampleType, ESPMode::ThreadSafe>& OutSample, bool bReverse)
+	{
+		// Code below assumes a fully specified range, no open bounds!
+		check(TimeRange.HasLowerBound() && TimeRange.HasUpperBound());
 
-			if (TimeRange.Overlaps(SampleTimeRange))
-			{
-				// Sample is at least partially inside the requested range, recall the range of samples we find...
-				if (FirstPossibleIndex < 0)
-				{
-					FirstPossibleIndex = Idx;
-				}
-				LastPossibleIndex = Idx;
-			}
-			else
-			{
-				if (!bReverse ? (SampleTimeRange.GetLowerBoundValue() >= TimeRange.GetUpperBoundValue()) :
-								(SampleTimeRange.GetUpperBoundValue() <= TimeRange.GetLowerBoundValue()))
-				{
-					// Sample is entirely past requested time range, we can stop
-					// (we assume monotonically increasing time stamps here)
-					break;
-				}
+		OutSample.Reset();
 
-				// If the incoming data it not monotonically increasing we migth get here after we already found the first overlapping sample
-				// -> we do not count further non-overlapping, older samples into this range
-				if (FirstPossibleIndex < 0)
-				{
-					// Sample is before time range, we will delete is later, no reason to keep it
-					++NumOldSamplesAtBegin;
-				}
-				else
-				{
-					// If we find an older non-verlapping sample after an overlapping one, we move the last possible index on to ensure these samples die ASAP
-					LastPossibleIndex = Idx;
-				}
-			}
-		}
+		FScopeLock Lock(&CriticalSection);
+
+		int32 FirstPossibleIndex, LastPossibleIndex, NumOldSamplesAtBegin;
+		FindRangeInQueue(TimeRange, bReverse, FirstPossibleIndex, LastPossibleIndex, NumOldSamplesAtBegin);
 
 		// Found anything?
 		if (FirstPossibleIndex >= 0)
@@ -192,7 +171,7 @@ public:
 
 					// Check once more if this sample is actually overlapping as we may get non-monotonically increasing data...
 					TRange<FMediaTimeStamp> SampleTimeRange = !bReverse ? TRange<FMediaTimeStamp>(Sample->GetTime(), Sample->GetTime() + Sample->GetDuration())
-																		: TRange<FMediaTimeStamp>(Sample->GetTime() - Sample->GetDuration(), Sample->GetTime());
+						: TRange<FMediaTimeStamp>(Sample->GetTime() - Sample->GetDuration(), Sample->GetTime());
 
 					if (TimeRange.Overlaps(SampleTimeRange))
 					{
@@ -239,20 +218,20 @@ public:
 	}
 
 
-	uint32 PurgeOutdatedSamples(const FMediaTimeStamp & ReferenceTime, bool bReversed)
+	uint32 PurgeOutdatedSamples(const FMediaTimeStamp& ReferenceTime, bool bReversed, FTimespan MaxAge)
 	{
 		FScopeLock Lock(&CriticalSection);
 
 		int32 Num = Samples.Num();
 		if (Num > 0)
 		{
-			int32 Idx = 0;
+			// All samples at or beyond the reference time are good to stay
+			int32 Idx;
 			if (!bReversed)
 			{
-				for (; Idx < Num; ++Idx)
+				for (Idx = Num - 1; Idx >= 0; --Idx)
 				{
-					const TSharedPtr<SampleType, ESPMode::ThreadSafe> & Sample = Samples[Idx];
-					if ((Sample->GetTime() + Sample->GetDuration()) > ReferenceTime)
+					if (Samples[Idx]->GetTime() < ReferenceTime)
 					{
 						break;
 					}
@@ -260,22 +239,28 @@ public:
 			}
 			else
 			{
-				for (; Idx < Num; ++Idx)
+				for (Idx = Num - 1; Idx >= 0; --Idx)
 				{
-					const TSharedPtr<SampleType, ESPMode::ThreadSafe> & Sample = Samples[Idx];
-					if ((Sample->GetTime() - Sample->GetDuration()) < ReferenceTime)
+					if (Samples[Idx]->GetTime() > ReferenceTime)
 					{
 						break;
 					}
 				}
 			}
-			if (Idx > 0)
+			// Accumulate durations of samples from the reference time backwards to judge what to purge as "too old"
+			FTimespan Age = FTimespan::Zero();
+			for (; Idx >= 0; --Idx)
 			{
-				Samples.RemoveAt(0, Idx);
+				auto Duration = Samples[Idx]->GetDuration();
+				Age += Duration;
+				if (Age > MaxAge)
+				{
+					// All earlier samples, including the current one are "too old"
+					Samples.RemoveAt(0, Idx + 1);
+					return Idx + 1;
+				}
 			}
-			return Idx;
 		}
-
 		return 0;
 	}
 
@@ -298,6 +283,58 @@ public:
 	}
 
 protected:
+	void FindRangeInQueue(const TRange<FMediaTimeStamp>& TimeRange, bool bReverse, int32& FirstPossibleIndex, int32& LastPossibleIndex, int32& NumOldSamplesAtBegin)
+	{
+		FirstPossibleIndex = -1;
+		LastPossibleIndex = -1;
+		NumOldSamplesAtBegin = 0;
+
+		int32 Num = Samples.Num();
+		if (Num > 0)
+		{
+			for (int32 Idx = 0; Idx < Num; ++Idx)
+			{
+				const TSharedPtr<SampleType, ESPMode::ThreadSafe>& Sample = Samples[Idx];
+				TRange<FMediaTimeStamp> SampleTimeRange = !bReverse ? TRange<FMediaTimeStamp>(Sample->GetTime(), Sample->GetTime() + Sample->GetDuration())
+					: TRange<FMediaTimeStamp>(Sample->GetTime() - Sample->GetDuration(), Sample->GetTime());
+
+				if (TimeRange.Overlaps(SampleTimeRange))
+				{
+					// Sample is at least partially inside the requested range, recall the range of samples we find...
+					if (FirstPossibleIndex < 0)
+					{
+						FirstPossibleIndex = Idx;
+					}
+					LastPossibleIndex = Idx;
+				}
+				else
+				{
+					if (!bReverse ? (SampleTimeRange.GetLowerBoundValue() >= TimeRange.GetUpperBoundValue()) :
+						(SampleTimeRange.GetUpperBoundValue() <= TimeRange.GetLowerBoundValue()))
+					{
+						// Sample is entirely past requested time range, we can stop
+						// (we assume monotonically increasing time stamps here)
+						break;
+					}
+
+					// If the incoming data it not monotonically increasing we might get here after we already found the first overlapping sample
+					// -> we do not count further non-overlapping, older samples into this range
+					if (FirstPossibleIndex < 0)
+					{
+						// Sample is before time range, we will delete is later, no reason to keep it
+						++NumOldSamplesAtBegin;
+					}
+					else
+					{
+						// If we find an older non-overlapping sample after an overlapping one, we move the last possible index on to ensure these samples die ASAP
+						LastPossibleIndex = Idx;
+					}
+				}
+			}
+		}
+
+	}
+
 	mutable FCriticalSection CriticalSection;
 	TArray<TSharedPtr<SampleType, ESPMode::ThreadSafe>> Samples;
 };
diff --git a/Engine/Source/Runtime/MediaUtils/Public/MediaSamples.h b/Engine/Source/Runtime/MediaUtils/Public/MediaSamples.h
--- a/Engine/Source/Runtime/MediaUtils/Public/MediaSamples.h
+++ b/Engine/Source/Runtime/MediaUtils/Public/MediaSamples.h
@@ -135,10 +135,6 @@ public:
 		return VideoSampleQueue.Num();
 	}
 
-	/**
-	 * Peek next video sample's timestamp
-	 * @return true if value could be retrieved, false otherwise
-	 */
 public:
 
 	//~ IMediaSamples interface
@@ -152,11 +148,16 @@ public:
 
 	virtual bool FetchAudio(TRange<FMediaTimeStamp> TimeRange, TSharedPtr<IMediaAudioSample, ESPMode::ThreadSafe>& OutSample) override;
 	virtual bool FetchCaption(TRange<FMediaTimeStamp> TimeRange, TSharedPtr<IMediaOverlaySample, ESPMode::ThreadSafe>& OutSample) override;
+	virtual bool FetchMetadata(TRange<FMediaTimeStamp> TimeRange, TSharedPtr<IMediaBinarySample, ESPMode::ThreadSafe>& OutSample) override;
 	virtual bool FetchSubtitle(TRange<FMediaTimeStamp> TimeRange, TSharedPtr<IMediaOverlaySample, ESPMode::ThreadSafe>& OutSample) override;
 	virtual bool FetchVideo(TRange<FMediaTimeStamp> TimeRange, TSharedPtr<IMediaTextureSample, ESPMode::ThreadSafe>& OutSample) override;
 
 	virtual EFetchBestSampleResult FetchBestVideoSampleForTimeRange(const TRange<FMediaTimeStamp> & TimeRange, TSharedPtr<IMediaTextureSample, ESPMode::ThreadSafe>& OutSample, bool bReverse) override;
 
+	/**
+	 * Peek next video sample's timestamp
+	 * @return true if value could be retrieved, false otherwise
+	 */
 	virtual bool PeekVideoSampleTime(FMediaTimeStamp & TimeStamp) override
 	{
 		TSharedPtr<IMediaTextureSample, ESPMode::ThreadSafe> Sample;
@@ -168,8 +169,35 @@ public:
 		return true;
 	}
 
-	virtual uint32 PurgeOutdatedVideoSamples(const FMediaTimeStamp & ReferenceTime, bool bReversed) override;
-	virtual uint32 PurgeOutdatedSubtitleSamples(const FMediaTimeStamp & ReferenceTime, bool bReversed) override;
+	virtual bool DiscardVideoSamples(const TRange<FMediaTimeStamp>& TimeRange, bool bReverse) override
+	{
+		return VideoSampleQueue.Discard(TimeRange, bReverse);
+	}
+
+	virtual bool DiscardAudioSamples(const TRange<FMediaTimeStamp>& TimeRange, bool bReverse) override
+	{
+		return AudioSampleQueue.Discard(TimeRange, bReverse);
+	}
+
+	virtual bool DiscardCaptionSamples(const TRange<FMediaTimeStamp>& TimeRange, bool bReverse) override
+	{
+		return CaptionSampleQueue.Discard(TimeRange, bReverse);
+	}
+
+	virtual bool DiscardSubtitleSamples(const TRange<FMediaTimeStamp>& TimeRange, bool bReverse) override
+	{
+		return SubtitleSampleQueue.Discard(TimeRange, bReverse);
+	}
+
+	virtual bool DiscardMetadataSamples(const TRange<FMediaTimeStamp>& TimeRange, bool bReverse) override
+	{
+		return MetadataSampleQueue.Discard(TimeRange, bReverse);
+	}
+
+	virtual uint32 PurgeOutdatedVideoSamples(const FMediaTimeStamp & ReferenceTime, bool bReversed, FTimespan MaxAge) override;
+	virtual uint32 PurgeOutdatedSubtitleSamples(const FMediaTimeStamp & ReferenceTime, bool bReversed, FTimespan MaxAge) override;
+	virtual uint32 PurgeOutdatedCaptionSamples(const FMediaTimeStamp& ReferenceTime, bool bReversed, FTimespan MaxAge) override;
+	virtual uint32 PurgeOutdatedMetadataSamples(const FMediaTimeStamp& ReferenceTime, bool bReversed, FTimespan MaxAge) override;
 
 	virtual bool CanReceiveVideoSamples(uint32 Num) const override;
 	virtual bool CanReceiveAudioSamples(uint32 Num) const override;
