Adding guards to avoid a division by zero when media facade code is working with duration data returned by plugin players
#rnx

diff --git a/Engine/Source/Runtime/MediaUtils/Private/MediaPlayerFacade.cpp b/Engine/Source/Runtime/MediaUtils/Private/MediaPlayerFacade.cpp
--- a/Engine/Source/Runtime/MediaUtils/Private/MediaPlayerFacade.cpp
+++ b/Engine/Source/Runtime/MediaUtils/Private/MediaPlayerFacade.cpp
@@ -2633,18 +2633,21 @@ bool FMediaPlayerFacade::GetCurrentPlaybackTimeRange(TRange<FMediaTimeStamp>& Ti
 	// (we do not clamp in the non-looping case as the rest of the code should deal with that fine)
 	if (Player->GetControls().IsLooping())
 	{
-		const FTimespan Duration = Player->GetControls().GetDuration();
-		FTimespan WrappedStart = WrappedModulo(TimeRange.GetLowerBoundValue().Time, Duration);
-		FTimespan WrappedEnd = WrappedModulo(TimeRange.GetUpperBoundValue().Time, Duration);
-		if (WrappedStart > WrappedEnd)
+		if (Player->GetControls().GetDuration() > 0.0f)
 		{
-			if (WrappedStart != TimeRange.GetLowerBoundValue().Time)
+			const FTimespan Duration = Player->GetControls().GetDuration();
+			FTimespan WrappedStart = WrappedModulo(TimeRange.GetLowerBoundValue().Time, Duration);
+			FTimespan WrappedEnd = WrappedModulo(TimeRange.GetUpperBoundValue().Time, Duration);
+			if (WrappedStart > WrappedEnd)
 			{
-				TimeRange.SetLowerBoundValue(FMediaTimeStamp(WrappedStart, FMediaTimeStamp::AdjustSecondaryIndex(TimeRange.GetLowerBoundValue().SequenceIndex, -1)));
-			}
-			if (WrappedEnd != TimeRange.GetUpperBoundValue().Time)
-			{
-				TimeRange.SetUpperBoundValue(FMediaTimeStamp(WrappedEnd, FMediaTimeStamp::AdjustSecondaryIndex(TimeRange.GetUpperBoundValue().SequenceIndex, 1)));
+				if (WrappedStart != TimeRange.GetLowerBoundValue().Time)
+				{
+					TimeRange.SetLowerBoundValue(FMediaTimeStamp(WrappedStart, FMediaTimeStamp::AdjustSecondaryIndex(TimeRange.GetLowerBoundValue().SequenceIndex, -1)));
+				}
+				if (WrappedEnd != TimeRange.GetUpperBoundValue().Time)
+				{
+					TimeRange.SetUpperBoundValue(FMediaTimeStamp(WrappedEnd, FMediaTimeStamp::AdjustSecondaryIndex(TimeRange.GetUpperBoundValue().SequenceIndex, 1)));
+				}
 			}
 		}
 	}
