Revising MFW's player facade's track selection logic to filter selection dynamically against the available sinks
(this enables players to optimize things like loop-behavior according to actual demanded data output)

diff --git a/Engine/Source/Runtime/Media/Public/IMediaTracks.h b/Engine/Source/Runtime/Media/Public/IMediaTracks.h
--- a/Engine/Source/Runtime/Media/Public/IMediaTracks.h
+++ b/Engine/Source/Runtime/Media/Public/IMediaTracks.h
@@ -35,7 +35,9 @@ enum class EMediaTrackType
 	Text,
 
 	/** Video track. */
-	Video
+	Video,
+
+	Num
 };
 
 
diff --git a/Engine/Source/Runtime/MediaUtils/Private/MediaPlayerFacade.cpp b/Engine/Source/Runtime/MediaUtils/Private/MediaPlayerFacade.cpp
--- a/Engine/Source/Runtime/MediaUtils/Private/MediaPlayerFacade.cpp
+++ b/Engine/Source/Runtime/MediaUtils/Private/MediaPlayerFacade.cpp
@@ -125,6 +125,8 @@ FMediaPlayerFacade::FMediaPlayerFacade()
 
 	MediaModule = FModuleManager::LoadModulePtr<IMediaModule>("Media");
 	bDidRecentPlayerHaveError = false;
+
+	ResetTracks();
 }
 
 
@@ -398,17 +400,6 @@ float FMediaPlayerFacade::GetRate() const
 }
 
 
-int32 FMediaPlayerFacade::GetSelectedTrack(EMediaTrackType TrackType) const
-{
-	TSharedPtr<IMediaPlayer, ESPMode::ThreadSafe> CurrentPlayer(Player);
-	if (!CurrentPlayer.IsValid())
-	{
-		return INDEX_NONE;
-	}
-	return CurrentPlayer->GetTracks().GetSelectedTrack((EMediaTrackType)TrackType);
-}
-
-
 FString FMediaPlayerFacade::GetStats() const
 {
 	TSharedPtr<IMediaPlayer, ESPMode::ThreadSafe> CurrentPlayer(Player);
@@ -1114,6 +1105,8 @@ bool FMediaPlayerFacade::ContinueOpen(IMediaPlayerLifecycleManagerDelegate::ICon
 		SeekIndex = 0;
 	}
 
+	ResetTracks();
+
 	if (bCreateNewPlayer)
 	{
 		NotifyLifetimeManagerDelegate_PlayerCreated();
@@ -1206,21 +1199,6 @@ bool FMediaPlayerFacade::Seek(const FTimespan& InTime)
 }
 
 
-bool FMediaPlayerFacade::SelectTrack(EMediaTrackType TrackType, int32 TrackIndex)
-{
-	if (!Player.IsValid() || !Player->GetTracks().SelectTrack((EMediaTrackType)TrackType, TrackIndex))
-	{
-		return false;
-	}
-
-	if (!Player->GetPlayerFeatureFlag(IMediaPlayer::EFeatureFlag::IsTrackSwitchSeamless))
-	{
-		Flush();
-	}
-	return true;
-}
-
-
 void FMediaPlayerFacade::SetBlockOnTime(const FTimespan& Time)
 {
 #if !MEDIAPLAYERFACADE_DISABLE_BLOCKING
@@ -1902,26 +1880,110 @@ void FMediaPlayerFacade::ProcessEvent(EMediaEvent Event, bool bIsBroadcastAllowe
 }
 
 
-void FMediaPlayerFacade::SelectDefaultTracks()
+void FMediaPlayerFacade::ResetTracks()
 {
-	if (!Player.IsValid())
+	for (int32 Idx = 0; Idx < (int32)EMediaTrackType::Num; ++Idx)
 	{
-		return;
+		TrackSelection.UserSelection[Idx] = -1;
+		TrackSelection.PlayerSelection[Idx] = -1;
 	}
+}
 
-	IMediaTracks& Tracks = Player->GetTracks();
 
+void FMediaPlayerFacade::SelectDefaultTracks()
+{
 	FMediaPlayerTrackOptions TrackOptions;
 	if (ActivePlayerOptions.IsSet())
 	{
 		TrackOptions = ActivePlayerOptions.GetValue().Tracks;
 	}
 
-	Tracks.SelectTrack(EMediaTrackType::Audio, TrackOptions.Audio);
-	Tracks.SelectTrack(EMediaTrackType::Caption, TrackOptions.Caption);
-	Tracks.SelectTrack(EMediaTrackType::Metadata, TrackOptions.Metadata);
-	Tracks.SelectTrack(EMediaTrackType::Subtitle, TrackOptions.Subtitle);
-	Tracks.SelectTrack(EMediaTrackType::Video, TrackOptions.Video);
+	TrackSelection.UserSelection[(int32)EMediaTrackType::Audio] = TrackOptions.Audio;
+	TrackSelection.UserSelection[(int32)EMediaTrackType::Caption] = TrackOptions.Caption;
+	TrackSelection.UserSelection[(int32)EMediaTrackType::Metadata] = TrackOptions.Metadata;
+	TrackSelection.UserSelection[(int32)EMediaTrackType::Subtitle] = TrackOptions.Subtitle;
+	TrackSelection.UserSelection[(int32)EMediaTrackType::Video] = TrackOptions.Video;
+}
+
+
+bool FMediaPlayerFacade::SelectTrack(EMediaTrackType TrackType, int32 TrackIndex)
+{
+	TSharedPtr<IMediaPlayer, ESPMode::ThreadSafe> CurrentPlayer(Player);
+	if (CurrentPlayer.IsValid())
+	{
+		IMediaTracks& Tracks = CurrentPlayer->GetTracks();
+
+		if (Tracks.GetNumTracks(TrackType) > TrackIndex)
+		{
+			TrackSelection.UserSelection[(int32)TrackType] = TrackIndex;
+			return true;
+		}
+	}
+	return false;
+}
+
+
+int32 FMediaPlayerFacade::GetSelectedTrack(EMediaTrackType TrackType) const
+{
+	return TrackSelection.UserSelection[(int32)TrackType];
+}
+
+
+void FMediaPlayerFacade::UpdateTrackSelectionWithPlayer()
+{
+	check(Player.IsValid());
+
+	bool bChanges = false;
+
+	IMediaTracks& Tracks = Player->GetTracks();
+	for (int32 Idx = 0; Idx < (int32)EMediaTrackType::Num; ++Idx)
+	{
+		// Player and user selection are different?
+		if (TrackSelection.PlayerSelection[Idx] != TrackSelection.UserSelection[Idx])
+		{
+			// Yes...
+			int32 UserSelection = TrackSelection.UserSelection[Idx];
+
+			// Filter selection against the configured sinks...
+			if (UserSelection != -1)
+			{
+				if ((Idx == (int)EMediaTrackType::Audio && !PrimaryAudioSink.IsValid()) ||
+					(Idx == (int)EMediaTrackType::Video && VideoSampleSinks.IsEmpty()) ||
+					(Idx == (int)EMediaTrackType::Caption && CaptionSampleSinks.IsEmpty()) ||
+					(Idx == (int)EMediaTrackType::Subtitle && SubtitleSampleSinks.IsEmpty()) ||
+					(Idx == (int)EMediaTrackType::Metadata && MetadataSampleSinks.IsEmpty()))
+				{
+					UserSelection = -1;
+				}
+			}
+
+			// After filtering the user's selection, do we still have to change things?
+			if (TrackSelection.PlayerSelection[Idx] != UserSelection)
+			{
+				// Yes!
+				if (Tracks.SelectTrack((EMediaTrackType)Idx, UserSelection))
+				{
+					// Recall what is now selected with the player...
+					TrackSelection.PlayerSelection[Idx] = UserSelection;
+
+					bChanges = true;
+				}
+				else
+				{
+					// Track selection failed. Patch the user selection to be what we know of the player's, so we do not reattempt this over and over...
+					TrackSelection.UserSelection[Idx] = TrackSelection.PlayerSelection[Idx];
+				}
+			}
+		}
+	}
+
+	if (bChanges)
+	{
+		if (!Player->GetPlayerFeatureFlag(IMediaPlayer::EFeatureFlag::IsTrackSwitchSeamless))
+		{
+			Flush();
+		}
+	}
 }
 
 
@@ -1940,6 +2002,7 @@ void FMediaPlayerFacade::TickInput(FTimespan DeltaTime, FTimespan Timecode)
 
 	if (Player.IsValid())
 	{
+		UpdateTrackSelectionWithPlayer();
 		MonitorAudioEnablement();
 
 		Player->TickInput(DeltaTime, Timecode);
diff --git a/Engine/Source/Runtime/MediaUtils/Public/MediaPlayerFacade.h b/Engine/Source/Runtime/MediaUtils/Public/MediaPlayerFacade.h
--- a/Engine/Source/Runtime/MediaUtils/Public/MediaPlayerFacade.h
+++ b/Engine/Source/Runtime/MediaUtils/Public/MediaPlayerFacade.h
@@ -12,6 +12,7 @@
 #include "IMediaPlayerLifecycleManager.h"
 #include "IMediaTickable.h"
 #include "IMediaTimeSource.h"
+#include "IMediaTracks.h"
 #include "Internationalization/Text.h"
 #include "Math/MathFwd.h"
 #include "Math/Quat.h"
@@ -795,6 +796,19 @@ protected:
 	 **/
 	void ProcessEvent(EMediaEvent Event, bool bIsBroadcastAllowed);
 
+private:
+	struct FTrackSelection
+	{
+		int32 UserSelection[(int32)EMediaTrackType::Num];
+		int32 PlayerSelection[(int32)EMediaTrackType::Num];
+	} TrackSelection;
+
+	/** Reset all tracks. */
+	void ResetTracks();
+
+	/** Setup track selection with player. */
+	void UpdateTrackSelectionWithPlayer();
+
 	/** Select the default media tracks. */
 	void SelectDefaultTracks();
 
diff --git a/Engine/Source/Runtime/MediaUtils/Public/MediaSampleSinks.h b/Engine/Source/Runtime/MediaUtils/Public/MediaSampleSinks.h
--- a/Engine/Source/Runtime/MediaUtils/Public/MediaSampleSinks.h
+++ b/Engine/Source/Runtime/MediaUtils/Public/MediaSampleSinks.h
@@ -103,6 +103,16 @@ public:
 		return Sinks.Num();
 	}
 
+	/**
+	* Check if any sinks are inserted
+	* 
+	* @return True if empty, false otherwise
+	*/
+	bool IsEmpty() const
+	{
+		return Sinks.IsEmpty();
+	}
+
 	/**
 	 * Remove the given media sample sink from the collection.
 	 *
