Fixing issue with videos without audio that would make the clock already tick after open instead of waiting for a valid non-zero rate first.
#rnx

diff --git a/Engine/Source/Runtime/MediaUtils/Private/MediaPlayerFacade.cpp b/Engine/Source/Runtime/MediaUtils/Private/MediaPlayerFacade.cpp
--- a/Engine/Source/Runtime/MediaUtils/Private/MediaPlayerFacade.cpp
+++ b/Engine/Source/Runtime/MediaUtils/Private/MediaPlayerFacade.cpp
@@ -2452,9 +2452,10 @@ void FMediaPlayerFacade::PreSampleProcessingTimeHandling()
 
 	UpdateSeekStatus();
 
-	if (!SeekTargetTime.IsValid())
+	// No seeking and not paused either?
+	if (!SeekTargetTime.IsValid() && CurrentRate != 0.0f)
 	{
-		// No seek pending. Can we / Do we need to prime a non-audio clock?
+		// No seek pending & not paused. Can we / Do we need to prime a non-audio clock?
 		if (!bHaveActiveAudio && !BlockOnRange.IsSet() && !NextEstVideoTimeAtFrameStart.IsValid())
 		{
 			FMediaTimeStamp VideoTimeStamp;
