Fixing up recent fix for playback behavior of MFW player facade between open & play calls. (Fixes issue by properly resetting the current rate from previous sessions)
#rnx

diff --git a/Engine/Source/Runtime/MediaUtils/Private/MediaPlayerFacade.cpp b/Engine/Source/Runtime/MediaUtils/Private/MediaPlayerFacade.cpp
--- a/Engine/Source/Runtime/MediaUtils/Private/MediaPlayerFacade.cpp
+++ b/Engine/Source/Runtime/MediaUtils/Private/MediaPlayerFacade.cpp
@@ -1,5 +1,4 @@
 // Copyright Epic Games, Inc. All Rights Reserved.
-// Copyright Epic Games, Inc. All Rights Reserved.
 
 #include "MediaPlayerFacade.h"
 #include "MediaUtilsPrivate.h"
@@ -274,6 +273,7 @@ void FMediaPlayerFacade::Close()
 	Cache->Empty();
 	CurrentUrl.Empty();
 	LastRate = 0.0f;
+	CurrentRate = 0.0f;
 
 	bHaveActiveAudio = false;
 	VideoSampleAvailability = -1;
@@ -2452,8 +2452,8 @@ void FMediaPlayerFacade::PreSampleProcessingTimeHandling()
 
 	UpdateSeekStatus();
 
-	// No seeking and not paused either?
-	if (!SeekTargetTime.IsValid() && CurrentRate != 0.0f)
+	// No seeking?
+	if (!SeekTargetTime.IsValid())
 	{
 		// No seek pending & not paused. Can we / Do we need to prime a non-audio clock?
 		if (!bHaveActiveAudio && !BlockOnRange.IsSet() && !NextEstVideoTimeAtFrameStart.IsValid())
