MediaUtils: Fixed an issue with seeking due to a wrong overloaded version of Clamp() being used.

#rnx

diff --git a/Engine/Source/Runtime/MediaUtils/Private/MediaPlayerFacade.cpp b/Engine/Source/Runtime/MediaUtils/Private/MediaPlayerFacade.cpp
--- a/Engine/Source/Runtime/MediaUtils/Private/MediaPlayerFacade.cpp
+++ b/Engine/Source/Runtime/MediaUtils/Private/MediaPlayerFacade.cpp
@@ -1177,7 +1177,7 @@ bool FMediaPlayerFacade::Seek(const FTimespan& InTime)
 	}
 	else
 	{
-		Time = FTimespan(FMath::Clamp(InTime.GetTicks(), 0, Duration.GetTicks()));
+		Time = FTimespan(FMath::Clamp(InTime.GetTicks(), (int64)0L, Duration.GetTicks()));
 	}
 
 	if (!CurrentPlayer->GetControls().Seek(Time))
