Adding "sink events" to MediaFramework  so that sinks can track (some) player state changes (base work for some tick-reduction in the future)
#rnx

diff --git a/Engine/Source/Runtime/Media/Public/IMediaSamples.h b/Engine/Source/Runtime/Media/Public/IMediaSamples.h
--- a/Engine/Source/Runtime/Media/Public/IMediaSamples.h
+++ b/Engine/Source/Runtime/Media/Public/IMediaSamples.h
@@ -142,6 +142,9 @@ public:
 
 	virtual bool CanReceiveVideoSamples(uint32 Num) const { return true; }
 	virtual bool CanReceiveAudioSamples(uint32 Num) const { return true; }
+	virtual bool CanReceiveSubtitleSamples(uint32 Num) const { return true; }
+	virtual bool CanReceiveCaptionSamples(uint32 Num) const { return true; }
+	virtual bool CanReceiveMetadataSamples(uint32 Num) const { return true; }
 
 	virtual int32 NumAudio() const { return -1; }
 	virtual int32 NumCaption() const { return -1; }
diff --git a/Engine/Source/Runtime/MediaAssets/Private/Assets/MediaPlayer.cpp b/Engine/Source/Runtime/MediaAssets/Private/Assets/MediaPlayer.cpp
--- a/Engine/Source/Runtime/MediaAssets/Private/Assets/MediaPlayer.cpp
+++ b/Engine/Source/Runtime/MediaAssets/Private/Assets/MediaPlayer.cpp
@@ -63,7 +63,7 @@ UMediaPlayer::UMediaPlayer(const FObjectInitializer& ObjectInitializer)
 {
 	if (!HasAnyFlags(RF_ClassDefaultObject))
 	{
-		PlayerFacade = MakeShareable(new FMediaPlayerFacade());
+		PlayerFacade = MakeShareable(new FMediaPlayerFacade(this));
 		PlayerFacade->OnMediaEvent().AddUObject(this, &UMediaPlayer::HandlePlayerMediaEvent);
 	}
 }
diff --git a/Engine/Source/Runtime/MediaUtils/Private/MediaPlayerFacade.cpp b/Engine/Source/Runtime/MediaUtils/Private/MediaPlayerFacade.cpp
--- a/Engine/Source/Runtime/MediaUtils/Private/MediaPlayerFacade.cpp
+++ b/Engine/Source/Runtime/MediaUtils/Private/MediaPlayerFacade.cpp
@@ -108,7 +108,7 @@ static FTimespan WrappedModulo(FTimespan Time, FTimespan Duration)
 /* FMediaPlayerFacade structors
 *****************************************************************************/
 
-FMediaPlayerFacade::FMediaPlayerFacade()
+FMediaPlayerFacade::FMediaPlayerFacade(TWeakObjectPtr<UMediaPlayer> InMediaPlayer)
 	: TimeDelay(FTimespan::Zero())
 	, BlockOnRange(this)
 	, Cache(new FMediaSampleCache)
@@ -119,6 +119,7 @@ FMediaPlayerFacade::FMediaPlayerFacade()
 	, AudioSampleAvailability(-1)
 	, bIsSinkFlushPending(false)
 	, bAreEventsSafeForAnyThread(false)
+	, MediaPlayer(InMediaPlayer)
 {
 	BlockOnRangeDisabled = false;
 
@@ -131,6 +132,10 @@ FMediaPlayerFacade::FMediaPlayerFacade()
 
 FMediaPlayerFacade::~FMediaPlayerFacade()
 {
+	FMediaSampleSinkEventData Data;
+	Data.Detached.MediaPlayer = MediaPlayer.Get();
+	SendSinkEvent(EMediaSampleSinkEvent::Detached, Data);
+
 	if (Player.IsValid())
 	{
 		{
@@ -1449,6 +1454,11 @@ bool FMediaPlayerFacade::SetRate(float Rate)
 		return bRateOk;
 	}
 
+	// Notify sinks of rate change
+	FMediaSampleSinkEventData Data;
+	Data.PlaybackRateChanged.PlaybackRate = Rate;
+	SendSinkEvent(EMediaSampleSinkEvent::PlaybackRateChanged, Data);
+
 	if ((LastRate * Rate) < 0.0f)
 	{
 		Flush(); // direction change
@@ -1633,11 +1643,12 @@ void FMediaPlayerFacade::Flush(bool bExcludePlayer, bool bOnSeek)
 
 	FScopeLock Lock(&CriticalSection);
 
-	AudioSampleSinks.Flush();
-	CaptionSampleSinks.Flush();
-	MetadataSampleSinks.Flush();
-	SubtitleSampleSinks.Flush();
-	VideoSampleSinks.Flush();
+	auto RawMediaPlayer = MediaPlayer.Get();
+	AudioSampleSinks.Flush(RawMediaPlayer);
+	CaptionSampleSinks.Flush(RawMediaPlayer);
+	MetadataSampleSinks.Flush(RawMediaPlayer);
+	SubtitleSampleSinks.Flush(RawMediaPlayer);
+	VideoSampleSinks.Flush(RawMediaPlayer);
 
 	if (Player.IsValid() && !bExcludePlayer)
 	{
@@ -1670,6 +1681,21 @@ void FMediaPlayerFacade::Flush(bool bExcludePlayer, bool bOnSeek)
 }
 
 
+void FMediaPlayerFacade::SendSinkEvent(EMediaSampleSinkEvent Event, const FMediaSampleSinkEventData& Data)
+{
+	{
+	FScopeLock Lock(&CriticalSection);
+
+	AudioSampleSinks.ReceiveEvent(Event, Data);
+	MetadataSampleSinks.ReceiveEvent(Event, Data);
+	}
+
+	CaptionSampleSinks.ReceiveEvent(Event, Data);
+	SubtitleSampleSinks.ReceiveEvent(Event, Data);
+	VideoSampleSinks.ReceiveEvent(Event, Data);
+}
+
+
 bool FMediaPlayerFacade::GetAudioTrackFormat(int32 TrackIndex, int32 FormatIndex, FMediaAudioTrackFormat& OutFormat) const
 {
 	if (TrackIndex == INDEX_NONE)
@@ -1839,9 +1865,7 @@ void FMediaPlayerFacade::ProcessEvent(EMediaEvent Event, bool bIsBroadcastAllowe
 			UE_LOG(LogMediaUtils, Verbose, TEXT("PlayerFacade %p: Media Info:\n%s"), this, *MediaInfo);
 		}
 	}
-
-	if ((Event == EMediaEvent::PlaybackEndReached) ||
-		(Event == EMediaEvent::TracksChanged))
+	else if (Event == EMediaEvent::TracksChanged)
 	{
 		if (Player.IsValid() && !Player->GetPlayerFeatureFlag(IMediaPlayer::EFeatureFlag::UsePlaybackTimingV2))
 		{
@@ -1867,6 +1891,8 @@ void FMediaPlayerFacade::ProcessEvent(EMediaEvent Event, bool bIsBroadcastAllowe
 		if (CurrentUrl.IsEmpty())
 		{
 			// Yes, this also means: if we still have a player, it's still the one this event originated from
+			FMediaSampleSinkEventData Data;
+			SendSinkEvent(EMediaSampleSinkEvent::MediaClosed, Data);
 
 			// If player allows: close it down all the way right now
 			if (Player.IsValid() && Player->GetPlayerFeatureFlag(IMediaPlayer::EFeatureFlag::AllowShutdownOnClose))
@@ -1879,6 +1905,16 @@ void FMediaPlayerFacade::ProcessEvent(EMediaEvent Event, bool bIsBroadcastAllowe
 			MediaModule->GetTicker().RemoveTickable(AsShared());
 		}
 	}
+	else if (Event == EMediaEvent::PlaybackEndReached)
+	{
+		if (Player.IsValid() && !Player->GetPlayerFeatureFlag(IMediaPlayer::EFeatureFlag::UsePlaybackTimingV2))
+		{
+			// Execute flush for older players only
+			Flush();
+		}
+		FMediaSampleSinkEventData Data;
+		SendSinkEvent(EMediaSampleSinkEvent::PlaybackEndReached, Data);
+	}
 
 	if (bIsBroadcastAllowed)
 	{
@@ -2976,7 +3012,7 @@ bool FMediaPlayerFacade::ProcessVideoSamples(IMediaSamples& Samples, const TRang
 		// Enqueue the sample to render
 		// (we use a queue to stay compatible with existing structure and older sinks - new sinks will read this single entry right away on the gamethread
 		//  and pass it along to rendering outside the queue)
-		bool bOk = VideoSampleSinks.Enqueue(Sample.ToSharedRef(), FMediaPlayerQueueDepths::MaxVideoSinkDepth);
+		bool bOk = VideoSampleSinks.Enqueue(Sample.ToSharedRef());
 		check(bOk);
 
 		{
@@ -3010,7 +3046,7 @@ void FMediaPlayerFacade::ProcessCaptionSamples(IMediaSamples& Samples, const TRa
 	{
 		while (Samples.FetchCaption(TimeRange, Sample))
 		{
-			if (Sample.IsValid() && !CaptionSampleSinks.Enqueue(Sample.ToSharedRef(), FMediaPlayerQueueDepths::MaxCaptionSinkDepth))
+			if (Sample.IsValid() && !CaptionSampleSinks.Enqueue(Sample.ToSharedRef()))
 			{
 #if MEDIAPLAYERFACADE_TRACE_SINKOVERFLOWS
 				UE_LOG(LogMediaUtils, VeryVerbose, TEXT("PlayerFacade %p: Caption sample sink overflow"), this);
@@ -3040,7 +3076,7 @@ void FMediaPlayerFacade::ProcessSubtitleSamples(IMediaSamples& Samples, const TR
 		while (Samples.FetchSubtitle(TimeRange, Sample))
 		{
 			//UE_LOG(LogMediaUtils, Display, TEXT("Subtitle @%.3f: %s"), Sample->GetTime().Time.GetTotalSeconds(), *Sample->GetText().ToString());
-			if (Sample.IsValid() && !SubtitleSampleSinks.Enqueue(Sample.ToSharedRef(), FMediaPlayerQueueDepths::MaxSubtitleSinkDepth))
+			if (Sample.IsValid() && !SubtitleSampleSinks.Enqueue(Sample.ToSharedRef()))
 			{
 #if MEDIAPLAYERFACADE_TRACE_SINKOVERFLOWS
 				UE_LOG(LogMediaUtils, VeryVerbose, TEXT("PlayerFacade %p: Subtitle sample sink overflow"), this);
@@ -3069,7 +3105,7 @@ void FMediaPlayerFacade::ProcessMetadataSamples(IMediaSamples& Samples, const TR
 	{
 		while (Samples.FetchMetadata(TimeRange, Sample))
 		{
-			if (Sample.IsValid() && !MetadataSampleSinks.Enqueue(Sample.ToSharedRef(), FMediaPlayerQueueDepths::MaxMetadataSinkDepth))
+			if (Sample.IsValid() && !MetadataSampleSinks.Enqueue(Sample.ToSharedRef()))
 			{
 #if MEDIAPLAYERFACADE_TRACE_SINKOVERFLOWS
 				UE_LOG(LogMediaUtils, VeryVerbose, TEXT("PlayerFacade %p: Metadata sample sink overflow"), this);
@@ -3091,7 +3127,7 @@ void FMediaPlayerFacade::ProcessAudioSamplesV1(IMediaSamples& Samples, TRange<FT
 			continue;
 		}
 
-		if (!AudioSampleSinks.Enqueue(Sample.ToSharedRef(), FMediaPlayerQueueDepths::MaxAudioSinkDepth))
+		if (!AudioSampleSinks.Enqueue(Sample.ToSharedRef()))
 		{
 #if MEDIAPLAYERFACADE_TRACE_SINKOVERFLOWS
 			UE_LOG(LogMediaUtils, VeryVerbose, TEXT("PlayerFacade %p: Audio sample sink overflow"), this);
@@ -3134,7 +3170,7 @@ void FMediaPlayerFacade::ProcessVideoSamplesV1(IMediaSamples& Samples, TRange<FT
 
 		UE_LOG(LogMediaUtils, VeryVerbose, TEXT("PlayerFacade %p: Fetched video sample %s"), this, *Sample->GetTime().Time.ToString(TEXT("%h:%m:%s.%t")));
 
-		if (VideoSampleSinks.Enqueue(Sample.ToSharedRef(), FMediaPlayerQueueDepths::MaxVideoSinkDepth))
+		if (VideoSampleSinks.Enqueue(Sample.ToSharedRef()))
 		{
 			if (CurrentRate >= 0.0f)
 			{
@@ -3157,7 +3193,7 @@ void FMediaPlayerFacade::ProcessCaptionSamplesV1(IMediaSamples& Samples, TRange<
 
 	while (Samples.FetchCaption(TimeRange, Sample))
 	{
-		if (Sample.IsValid() && !CaptionSampleSinks.Enqueue(Sample.ToSharedRef(), FMediaPlayerQueueDepths::MaxCaptionSinkDepth))
+		if (Sample.IsValid() && !CaptionSampleSinks.Enqueue(Sample.ToSharedRef()))
 		{
 #if MEDIAPLAYERFACADE_TRACE_SINKOVERFLOWS
 			UE_LOG(LogMediaUtils, VeryVerbose, TEXT("PlayerFacade %p: Caption sample sink overflow"), this);
@@ -3173,7 +3209,7 @@ void FMediaPlayerFacade::ProcessSubtitleSamplesV1(IMediaSamples& Samples, TRange
 
 	while (Samples.FetchSubtitle(TimeRange, Sample))
 	{
-		if (Sample.IsValid() && !SubtitleSampleSinks.Enqueue(Sample.ToSharedRef(), FMediaPlayerQueueDepths::MaxSubtitleSinkDepth))
+		if (Sample.IsValid() && !SubtitleSampleSinks.Enqueue(Sample.ToSharedRef()))
 		{
 			FString Caption = Sample->GetText().ToString();
 			UE_LOG(LogMediaUtils, Log, TEXT("New caption @%.3f: %s"), Sample->GetTime().Time.GetTotalSeconds(), *Caption);
@@ -3192,7 +3228,7 @@ void FMediaPlayerFacade::ProcessMetadataSamplesV1(IMediaSamples& Samples, TRange
 
 	while (Samples.FetchMetadata(TimeRange, Sample))
 	{
-		if (Sample.IsValid() && !MetadataSampleSinks.Enqueue(Sample.ToSharedRef(), FMediaPlayerQueueDepths::MaxMetadataSinkDepth))
+		if (Sample.IsValid() && !MetadataSampleSinks.Enqueue(Sample.ToSharedRef()))
 		{
 #if MEDIAPLAYERFACADE_TRACE_SINKOVERFLOWS
 			UE_LOG(LogMediaUtils, VeryVerbose, TEXT("PlayerFacade %p: Metadata sample sink overflow"), this);
diff --git a/Engine/Source/Runtime/MediaUtils/Private/MediaSamples.cpp b/Engine/Source/Runtime/MediaUtils/Private/MediaSamples.cpp
--- a/Engine/Source/Runtime/MediaUtils/Private/MediaSamples.cpp
+++ b/Engine/Source/Runtime/MediaUtils/Private/MediaSamples.cpp
@@ -7,9 +7,6 @@
 #include "IMediaOverlaySample.h"
 #include "IMediaTextureSample.h"
 
-const uint32 MaxNumberOfQueuedVideoSamples = 4;
-const uint32 MaxNumberOfQueuedAudioSamples = 4;
-
 /* Local helpers
 *****************************************************************************/
 
@@ -63,6 +60,16 @@ bool FetchSample(TMediaSampleQueue<SampleType, SinkType>& SampleQueue, TRange<FM
 /* IMediaSamples interface
 *****************************************************************************/
 
+FMediaSamples::FMediaSamples(uint32 InMaxNumberOfQueuedAudioSamples, uint32 InMaxNumberOfQueuedVideoSamples, uint32 InMaxNumberOfQueuedCaptionSamples, uint32 InMaxNumberOfQueuedSubtitlesSamples, uint32 InMaxNumberOfQueuedMetaDataSamples)
+	: AudioSampleQueue(InMaxNumberOfQueuedAudioSamples)
+	, CaptionSampleQueue(InMaxNumberOfQueuedCaptionSamples)
+	, MetadataSampleQueue(InMaxNumberOfQueuedMetaDataSamples)
+	, SubtitleSampleQueue(InMaxNumberOfQueuedSubtitlesSamples)
+	, VideoSampleQueue(InMaxNumberOfQueuedVideoSamples)
+{
+}
+
+
 bool FMediaSamples::FetchAudio(TRange<FTimespan> TimeRange, TSharedPtr<IMediaAudioSample, ESPMode::ThreadSafe>& OutSample)
 {
 	return FetchSample(AudioSampleQueue, TimeRange, OutSample);
@@ -178,7 +185,7 @@ uint32 FMediaSamples::PurgeOutdatedMetadataSamples(const FMediaTimeStamp& Refere
  */
 bool FMediaSamples::CanReceiveVideoSamples(uint32 Num) const
 {
-	return (VideoSampleQueue.Num() + Num) <= MaxNumberOfQueuedVideoSamples;
+	return VideoSampleQueue.CanAcceptSamples(Num);
 }
 
 /**
@@ -186,5 +193,29 @@ bool FMediaSamples::CanReceiveVideoSamples(uint32 Num) const
  */
 bool FMediaSamples::CanReceiveAudioSamples(uint32 Num) const
 {
-	return (AudioSampleQueue.Num() + Num) <= MaxNumberOfQueuedAudioSamples;
+	return AudioSampleQueue.CanAcceptSamples(Num);
+}
+
+/**
+ * Check if can receive more subtitle samples
+ */
+bool  FMediaSamples::CanReceiveSubtitleSamples(uint32 Num) const
+{
+	return SubtitleSampleQueue.CanAcceptSamples(Num);
+}
+
+/**
+ * Check if can receive more caption samples
+ */
+bool  FMediaSamples::CanReceiveCaptionSamples(uint32 Num) const
+{
+	return CaptionSampleQueue.CanAcceptSamples(Num);
+}
+
+/**
+ * Check if can receive more metadata samples
+ */
+bool  FMediaSamples::CanReceiveMetadataSamples(uint32 Num) const
+{
+	return MetadataSampleQueue.CanAcceptSamples(Num);
 }
diff --git a/Engine/Source/Runtime/MediaUtils/Public/MediaPlayerFacade.h b/Engine/Source/Runtime/MediaUtils/Public/MediaPlayerFacade.h
--- a/Engine/Source/Runtime/MediaUtils/Public/MediaPlayerFacade.h
+++ b/Engine/Source/Runtime/MediaUtils/Public/MediaPlayerFacade.h
@@ -80,7 +80,7 @@ public:
 public:
 
 	/** Default constructor. */
-	FMediaPlayerFacade();
+	FMediaPlayerFacade(TWeakObjectPtr<UMediaPlayer> InMediaPlayer);
 
 	/** Virtual destructor. */
 	virtual ~FMediaPlayerFacade();
@@ -867,6 +867,8 @@ private:
 
 	bool ContinueOpen(IMediaPlayerLifecycleManagerDelegate::IControlRef NewLifecycleManagerDelegateControl, const FString& Url, const IMediaOptions* Options, const FMediaPlayerOptions* PlayerOptions, IMediaPlayerFactory* PlayerFactory, TSharedPtr<IMediaPlayer, ESPMode::ThreadSafe> ReusedPlayer, bool bCreateNewPlayer, uint64 NewPlayerInstanceID);
 
+	void SendSinkEvent(EMediaSampleSinkEvent Event, const FMediaSampleSinkEventData& Data);
+
 	/** Audio sample sinks. */
 	TWeakPtr<FMediaAudioSampleSink, ESPMode::ThreadSafe> PrimaryAudioSink;
 	FMediaAudioSampleSinks AudioSampleSinks;
@@ -1022,6 +1024,9 @@ private:
 	/** Mediamodule we are working in */
 	IMediaModule* MediaModule;
 
+	/** UMediaplayer we are working with */
+	TWeakObjectPtr<UMediaPlayer> MediaPlayer;
+
 	/** Control interface for lifecycle manager delegate */
 	IMediaPlayerLifecycleManagerDelegate::IControlRef LifecycleManagerDelegateControl;
 };
diff --git a/Engine/Source/Runtime/MediaUtils/Public/MediaSampleQueue.h b/Engine/Source/Runtime/MediaUtils/Public/MediaSampleQueue.h
--- a/Engine/Source/Runtime/MediaUtils/Public/MediaSampleQueue.h
+++ b/Engine/Source/Runtime/MediaUtils/Public/MediaSampleQueue.h
@@ -31,8 +31,8 @@ class TMediaSampleQueue
 {
 public:
 
-	/** Default constructor. */
-	TMediaSampleQueue()
+	TMediaSampleQueue(int32 InMaxSamplesInQueue = -1)
+		: MaxSamplesInQueue(InMaxSamplesInQueue)
 	{ }
 
 	/** Virtual destructor. */
@@ -272,6 +272,11 @@ public:
 	{
 		FScopeLock Lock(&CriticalSection);
 
+		if ((MaxSamplesInQueue > 0) && (Samples.Num() >= MaxSamplesInQueue))
+		{
+			return false;
+		}
+
 		Samples.Push(Sample);
 		return true;
 	}
@@ -282,6 +287,11 @@ public:
 		Samples.Empty();
 	}
 
+	virtual bool CanAcceptSamples(int32 NumSamples) const override
+	{
+		return (MaxSamplesInQueue < 0) || ((Samples.Num() + NumSamples) <= MaxSamplesInQueue);
+	}
+
 protected:
 	void FindRangeInQueue(const TRange<FMediaTimeStamp>& TimeRange, bool bReverse, int32& FirstPossibleIndex, int32& LastPossibleIndex, int32& NumOldSamplesAtBegin)
 	{
@@ -337,6 +347,7 @@ protected:
 
 	mutable FCriticalSection CriticalSection;
 	TArray<TSharedPtr<SampleType, ESPMode::ThreadSafe>> Samples;
+	int32 MaxSamplesInQueue;
 };
 
 
@@ -344,8 +355,9 @@ protected:
 class FMediaAudioSampleQueue : public TMediaSampleQueue<class IMediaAudioSample, class FMediaAudioSampleSink>
 {
 public:
-	FMediaAudioSampleQueue(int32 InMaxAudioSamplesInQueue = -1)
-		: MaxAudioSamplesInQueue(InMaxAudioSamplesInQueue) {}
+	FMediaAudioSampleQueue(uint32 MaxSamplesInQueue = -1)
+		: TMediaSampleQueue<class IMediaAudioSample, class FMediaAudioSampleSink>(MaxSamplesInQueue)
+	{ }
 
 	void SetAudioTime(const FMediaTimeStampSample & InAudioTime)
 	{
@@ -371,26 +383,7 @@ public:
 		TMediaSampleQueue<class IMediaAudioSample, class FMediaAudioSampleSink>::RequestFlush();
 		AudioTime.Invalidate();
 	}
-
-	virtual bool Enqueue(const TSharedRef<IMediaAudioSample, ESPMode::ThreadSafe>& Sample) override
-	{
-		FScopeLock Lock(&CriticalSection);
-
-		if (MaxAudioSamplesInQueue > 0 && Samples.Num() >= MaxAudioSamplesInQueue)
-		{
-			return false;
-		}
-
-		Samples.Push(Sample);
-		return true;
-	}
-
-	virtual bool CanAcceptSamples(int32 NumSamples) const
-	{
-		return (Samples.Num() + NumSamples) <= MaxAudioSamplesInQueue;
-	}
 private:
-	int32 MaxAudioSamplesInQueue;
 	FMediaTimeStampSample AudioTime;
 };
 
diff --git a/Engine/Source/Runtime/MediaUtils/Public/MediaSampleQueueDepths.h b/Engine/Source/Runtime/MediaUtils/Public/MediaSampleQueueDepths.h
--- a/Engine/Source/Runtime/MediaUtils/Public/MediaSampleQueueDepths.h
+++ b/Engine/Source/Runtime/MediaUtils/Public/MediaSampleQueueDepths.h
@@ -4,6 +4,9 @@
 
 #include "CoreTypes.h"
 
+/** Legacy player output queue depth maximums
+ * @note These are not used inside MediaFramework and are meant to support older player sources code only.
+*/
 struct FMediaPlayerQueueDepths
 {
 	static const int32 MaxAudioSinkDepth = 512;
diff --git a/Engine/Source/Runtime/MediaUtils/Public/MediaSampleSink.h b/Engine/Source/Runtime/MediaUtils/Public/MediaSampleSink.h
--- a/Engine/Source/Runtime/MediaUtils/Public/MediaSampleSink.h
+++ b/Engine/Source/Runtime/MediaUtils/Public/MediaSampleSink.h
@@ -5,7 +5,62 @@
 #include "CoreTypes.h"
 #include "Misc/Timespan.h"
 #include "Templates/SharedPointer.h"
+#include "Delegates/Delegate.h"
 #include "IMediaTimeSource.h"
+#include "IMediaTracks.h"
+
+class UMediaPlayer;
+
+/** Events send from player to sink */
+enum class EMediaSampleSinkEvent
+{
+	Attached,				//!< Attached to a UMediaPlayer 
+	Detached,				//!< Detached from a UMediaPlayer
+	PlayerPluginChange,		//!< Player plugin used changed
+	SampleDataUpdate,		//!< Sample data was updated
+	FlushWasRequested,		//!< Flush has been requested
+	MediaClosed,			//!< Media has been closed
+	PlaybackEndReached,		//!< End of playback has been reached
+	PlaybackRateChanged,	//!< Rate of playback changed
+};
+
+struct FMediaSampleSinkEventData
+{
+	union {
+		struct
+		{
+			UMediaPlayer* MediaPlayer;
+		} Attached;
+		struct
+		{
+			UMediaPlayer* MediaPlayer;
+		} Detached;
+		struct
+		{
+			UMediaPlayer* MediaPlayer;
+		} PlayerPluginChange;
+		struct
+		{
+			uint32 Dummy;
+		} SampleDataUpdate;
+		struct
+		{
+			UMediaPlayer* MediaPlayer;
+		} FlushWasRequested;
+		struct
+		{
+			uint32 Dummy;
+		} MediaClosed;
+		struct
+		{
+			uint32 Dummy;
+		} PlaybackEndReached;
+		struct
+		{
+			float PlaybackRate;
+		} PlaybackRateChanged;
+	};
+};
 
 /**
  * Interface for media sample sinks.
@@ -36,6 +91,18 @@ public:
 	 */
 	virtual int32 Num() const = 0;
 
+	/**
+	 * Check if sink can accept new samples
+	 * 
+	 * @param NumSamples How many samples we would like the sink to accept
+	 * @return True if samples could be accepted, false otherwise
+	 * @note Override in implementation as needed
+	 */
+	virtual bool CanAcceptSamples(int32 NumSamples) const
+	{
+		return true;
+	}
+
 	/**
 	 * Request to flush the queue.
 	 *
@@ -44,10 +111,30 @@ public:
 	 */
 	virtual void RequestFlush() = 0;
 
+	/**
+	 * Receive event
+	 */
+	void ReceiveEvent(EMediaSampleSinkEvent Event, const FMediaSampleSinkEventData& Data)
+	{
+		MediaSampleSinkEvent.Broadcast(Event, Data);
+	}
+
+	/**
+	* Register to receive events flowing to this sink
+	*/
+	DECLARE_EVENT_TwoParams(TMediaSampleSink<SampleType>, FOnMediaSampleSinkEvent, EMediaSampleSinkEvent /*Event*/, const FMediaSampleSinkEventData& /*Data*/)
+	FOnMediaSampleSinkEvent& OnMediaSampleSinkEvent()
+	{
+		return MediaSampleSinkEvent;
+	}
+
 public:
 
 	/** Virtual destructor. */
 	virtual ~TMediaSampleSink() { }
+
+private:
+	FOnMediaSampleSinkEvent MediaSampleSinkEvent;
 };
 
 
@@ -60,8 +147,6 @@ public:
 	 */
 	virtual FMediaTimeStampSample GetAudioTime() const = 0;
 	virtual void InvalidateAudioTime() = 0;
-
-	virtual bool CanAcceptSamples(int32 NumSamples) const = 0;
 };
 
 /** Type definition for binary sample sink. */
diff --git a/Engine/Source/Runtime/MediaUtils/Public/MediaSampleSinks.h b/Engine/Source/Runtime/MediaUtils/Public/MediaSampleSinks.h
--- a/Engine/Source/Runtime/MediaUtils/Public/MediaSampleSinks.h
+++ b/Engine/Source/Runtime/MediaUtils/Public/MediaSampleSinks.h
@@ -25,9 +25,16 @@ public:
 	 * @param SampleSink The sink to add.
 	 * @see Num, Remove
 	 */
-	void Add(const TSharedRef<SinkType, ESPMode::ThreadSafe>& SampleSink)
+	void Add(const TSharedRef<SinkType, ESPMode::ThreadSafe>& SampleSink, UMediaPlayer* Player = nullptr)
 	{
-		Sinks.AddUnique(SampleSink);
+		if (Sinks.Find(SampleSink) != INDEX_NONE)
+		{
+			return;
+		}
+		Sinks.Add(SampleSink);
+		FMediaSampleSinkEventData Data;
+		Data.Attached = { Player };
+		SampleSink->ReceiveEvent(EMediaSampleSinkEvent::Attached, Data);
 	}
 
 	/**
@@ -40,7 +47,7 @@ public:
 	 * @return true if the sample was enqueued to all sinks, false if one or more sinks overflowed.
 	 * @see Flush
 	 */
-	bool Enqueue(const TSharedRef<SampleType, ESPMode::ThreadSafe>& Sample, int32 MaxDepth)
+	bool Enqueue(const TSharedRef<SampleType, ESPMode::ThreadSafe>& Sample)
 	{
 		bool Overflowed = false;
 
@@ -50,18 +57,20 @@ public:
 
 			if (Sink.IsValid())
 			{
-				if (Sink->Num() >= MaxDepth)
+				if (!Sink->CanAcceptSamples(1))
 				{
 					Overflowed = true;
 				}
 				else
 				{
 					Sink->Enqueue(Sample);
+					FMediaSampleSinkEventData Data;
+					Sink->ReceiveEvent(EMediaSampleSinkEvent::SampleDataUpdate, Data);
 				}
 			}
 			else
 			{
-				Sinks.RemoveAtSwap(SinkIndex);
+				Sinks.RemoveAt(SinkIndex);
 			}
 		}
 
@@ -75,7 +84,7 @@ public:
 	 *
 	 * @see Enqueue
 	 */
-	void Flush()
+	void Flush(UMediaPlayer* MediaPlayer)
 	{
 		for (int32 SinkIndex = Sinks.Num() - 1; SinkIndex >= 0; --SinkIndex)
 		{
@@ -84,10 +93,14 @@ public:
 			if (Sink.IsValid())
 			{
 				Sink->RequestFlush();
+
+				FMediaSampleSinkEventData Data;
+				Data.FlushWasRequested = { MediaPlayer };
+				Sink->ReceiveEvent(EMediaSampleSinkEvent::FlushWasRequested, Data);
 			}
 			else
 			{
-				Sinks.RemoveAtSwap(SinkIndex);
+				Sinks.RemoveAt(SinkIndex);
 			}
 		}
 	}
@@ -119,9 +132,14 @@ public:
 	 * @param SampleSink The sink to remove.
 	 * @see Add, Num
 	 */
-	void Remove(const TSharedRef<SinkType, ESPMode::ThreadSafe>& SampleSink)
+	void Remove(const TSharedRef<SinkType, ESPMode::ThreadSafe>& SampleSink, UMediaPlayer* Player = nullptr)
 	{
-		Sinks.Remove(SampleSink);
+		if (Sinks.RemoveSingle(SampleSink))
+		{
+			FMediaSampleSinkEventData Data;
+			Data.Detached = { Player };
+			SampleSink->ReceiveEvent(EMediaSampleSinkEvent::Detached, Data);
+		}
 	}
 
 	/**
@@ -133,7 +151,27 @@ public:
 		{
 			if (!Sinks[SinkIndex].IsValid())
 			{
-				Sinks.RemoveAtSwap(SinkIndex);
+				Sinks.RemoveAt(SinkIndex);
+			}
+		}
+	}
+
+	/**
+	 * Receive event and broadcast to sinks
+	 */
+	void ReceiveEvent(EMediaSampleSinkEvent Event, const FMediaSampleSinkEventData& Data)
+	{
+		for (int32 SinkIndex = Sinks.Num() - 1; SinkIndex >= 0; --SinkIndex)
+		{
+			TSharedPtr<SinkType, ESPMode::ThreadSafe> Sink = Sinks[SinkIndex].Pin();
+
+			if (Sink.IsValid())
+			{
+				Sink->ReceiveEvent(Event, Data);
+			}
+			else
+			{
+				Sinks.RemoveAt(SinkIndex);
 			}
 		}
 	}
diff --git a/Engine/Source/Runtime/MediaUtils/Public/MediaSamples.h b/Engine/Source/Runtime/MediaUtils/Public/MediaSamples.h
--- a/Engine/Source/Runtime/MediaUtils/Public/MediaSamples.h
+++ b/Engine/Source/Runtime/MediaUtils/Public/MediaSamples.h
@@ -24,6 +24,17 @@ class MEDIAUTILS_API FMediaSamples
 	: public IMediaSamples
 {
 public:
+	static constexpr uint32 MaxNumberOfQueuedVideoSamples = 4;
+	static constexpr uint32 MaxNumberOfQueuedAudioSamples = 4;
+	static constexpr uint32 MaxNumberOfQueuedCaptionSamples = 4;
+	static constexpr uint32 MaxNumberOfQueuedSubtitlesSamples = 4;
+	static constexpr uint32 MaxNumberOfQueuedMetaDataSamples = 4;
+
+	FMediaSamples(uint32 InMaxNumberOfQueuedAudioSamples = MaxNumberOfQueuedAudioSamples,
+				  uint32 InMaxNumberOfQueuedVideoSamples = MaxNumberOfQueuedVideoSamples,
+				  uint32 InMaxNumberOfQueuedCaptionSamples = MaxNumberOfQueuedCaptionSamples,
+				  uint32 InMaxNumberOfQueuedSubtitlesSamples = MaxNumberOfQueuedSubtitlesSamples,
+				  uint32 InMaxNumberOfQueuedMetaDataSamples = MaxNumberOfQueuedMetaDataSamples);
 
 	/**
 	 * Add the given audio sample to the cache.
@@ -201,6 +212,9 @@ public:
 
 	virtual bool CanReceiveVideoSamples(uint32 Num) const override;
 	virtual bool CanReceiveAudioSamples(uint32 Num) const override;
+	virtual bool CanReceiveSubtitleSamples(uint32 Num) const override;
+	virtual bool CanReceiveCaptionSamples(uint32 Num) const override;
+	virtual bool CanReceiveMetadataSamples(uint32 Num) const override;
 
 private:
 
