Fixing up event processing if/else double case
#rnx

diff --git a/Engine/Source/Runtime/MediaUtils/Private/MediaPlayerFacade.cpp b/Engine/Source/Runtime/MediaUtils/Private/MediaPlayerFacade.cpp
--- a/Engine/Source/Runtime/MediaUtils/Private/MediaPlayerFacade.cpp
+++ b/Engine/Source/Runtime/MediaUtils/Private/MediaPlayerFacade.cpp
@@ -1843,11 +1843,7 @@ void FMediaPlayerFacade::ProcessEvent(EMediaEvent Event, bool bIsBroadcastAllowe
 {
 	SCOPE_CYCLE_COUNTER(STAT_MediaUtils_FacadeProcessEvent);
 
-	if (Event == EMediaEvent::TracksChanged)
-	{
-		SelectDefaultTracks();
-	}
-	else if ((Event == EMediaEvent::MediaOpened) || (Event == EMediaEvent::MediaOpenFailed))
+	if ((Event == EMediaEvent::MediaOpened) || (Event == EMediaEvent::MediaOpenFailed))
 	{
 		if (Event == EMediaEvent::MediaOpenFailed)
 		{
@@ -1867,6 +1863,8 @@ void FMediaPlayerFacade::ProcessEvent(EMediaEvent Event, bool bIsBroadcastAllowe
 	}
 	else if (Event == EMediaEvent::TracksChanged)
 	{
+		SelectDefaultTracks();
+
 		if (Player.IsValid() && !Player->GetPlayerFeatureFlag(IMediaPlayer::EFeatureFlag::UsePlaybackTimingV2))
 		{
 			// Execute flush for older players only
