[Engine]

diff --git a/Engine/Source/Runtime/MediaUtils/Public/MediaPlayerFacade.h b/Engine/Source/Runtime/MediaUtils/Public/MediaPlayerFacade.h
--- a/Engine/Source/Runtime/MediaUtils/Public/MediaPlayerFacade.h
+++ b/Engine/Source/Runtime/MediaUtils/Public/MediaPlayerFacade.h
@@ -60,7 +60,7 @@ struct FMediaVideoTrackFormat;
  * related methods in this class allow for INDEX_NONE to be used as track and format
  * indices in order to indicate the 'current selection'.
  */
-class MEDIAUTILS_API FMediaPlayerFacade
+class FMediaPlayerFacade
 	: public IMediaClockSink
 	, public IMediaTickable
 	, protected IMediaEventSink
@@ -80,10 +80,10 @@ public:
 public:
 
 	/** Default constructor. */
-	FMediaPlayerFacade(TWeakObjectPtr<UMediaPlayer> InMediaPlayer);
+	MEDIAUTILS_API FMediaPlayerFacade(TWeakObjectPtr<UMediaPlayer> InMediaPlayer);
 
 	/** Virtual destructor. */
-	virtual ~FMediaPlayerFacade();
+	MEDIAUTILS_API virtual ~FMediaPlayerFacade();
 
 public:
 
@@ -93,7 +93,7 @@ public:
 	 * @param SampleSink The sink to receive audio samples.
 	 * @see  AddCaptionSampleSink, AddMetadataSampleSink, AddSubtitleSampleSink, AddVideoSampleSink
 	 */
-	void AddAudioSampleSink(const TSharedRef<FMediaAudioSampleSink, ESPMode::ThreadSafe>& SampleSink);
+	MEDIAUTILS_API void AddAudioSampleSink(const TSharedRef<FMediaAudioSampleSink, ESPMode::ThreadSafe>& SampleSink);
 
 	/**
 	 * Add the given audio sample sink to this player.
@@ -101,7 +101,7 @@ public:
 	 * @param SampleSink The sink to receive caption samples.
 	 * @see AddAudioSampleSink, AddMetadataSampleSink, AddSubtitleSampleSink, AddVideoSampleSink
 	 */
-	void AddCaptionSampleSink(const TSharedRef<FMediaOverlaySampleSink, ESPMode::ThreadSafe>& SampleSink);
+	MEDIAUTILS_API void AddCaptionSampleSink(const TSharedRef<FMediaOverlaySampleSink, ESPMode::ThreadSafe>& SampleSink);
 
 	/**
 	 * Add the given audio sample sink to this player.
@@ -109,7 +109,7 @@ public:
 	 * @param SampleSink The sink to receive metadata samples.
 	 * @see AddAudioSampleSink, AddCaptionSampleSink, AddSubtitleSampleSink, AddVideoSampleSink
 	 */
-	void AddMetadataSampleSink(const TSharedRef<FMediaBinarySampleSink, ESPMode::ThreadSafe>& SampleSink);
+	MEDIAUTILS_API void AddMetadataSampleSink(const TSharedRef<FMediaBinarySampleSink, ESPMode::ThreadSafe>& SampleSink);
 
 	/**
 	 * Add the given audio sample sink to this player.
@@ -117,7 +117,7 @@ public:
 	 * @param SampleSink The sink to receive subtitle samples.
 	 * @see AddAudioSampleSink, AddCaptionSampleSink, AddMetadataSampleSink, AddVideoSampleSink
 	 */
-	void AddSubtitleSampleSink(const TSharedRef<FMediaOverlaySampleSink, ESPMode::ThreadSafe>& SampleSink);
+	MEDIAUTILS_API void AddSubtitleSampleSink(const TSharedRef<FMediaOverlaySampleSink, ESPMode::ThreadSafe>& SampleSink);
 
 	/**
 	 * Add the given audio sample sink to this player.
@@ -125,7 +125,7 @@ public:
 	 * @param SampleSink The sink to receive video samples.
 	 * @see AddAudioSampleSink, AddCaptionSampleSink, AddMetadataSampleSink, AddSubtitleSampleSink
 	 */
-	void AddVideoSampleSink(const TSharedRef<FMediaTextureSampleSink, ESPMode::ThreadSafe>& SampleSink);
+	MEDIAUTILS_API void AddVideoSampleSink(const TSharedRef<FMediaTextureSampleSink, ESPMode::ThreadSafe>& SampleSink);
 
 	/**
 	 * Whether playback can be paused.
@@ -136,7 +136,7 @@ public:
 	 * @return true if pausing is allowed, false otherwise.
 	 * @see CanResume, CanScrub, CanSeek, Pause
 	 */
-	bool CanPause() const;
+	MEDIAUTILS_API bool CanPause() const;
 
 	/**
 	 * Whether the specified URL can be played by this player.
@@ -148,7 +148,7 @@ public:
 	 * @param Options Optional media parameters.
 	 * @see CanPlaySource, SetDesiredPlayerName
 	 */
-	bool CanPlayUrl(const FString& Url, const IMediaOptions* Options);
+	MEDIAUTILS_API bool CanPlayUrl(const FString& Url, const IMediaOptions* Options);
 
 	/**
 	 * Whether playback can be resumed.
@@ -159,7 +159,7 @@ public:
 	 * @return true if resuming is allowed, false otherwise.
 	 * @see CanPause, CanScrub, CanSeek, SetRate
 	 */
-	bool CanResume() const;
+	MEDIAUTILS_API bool CanResume() const;
 
 	/**
 	 * Whether playback can be scrubbed.
@@ -167,7 +167,7 @@ public:
 	 * @return true if scrubbing is allowed, false otherwise.
 	 * @see CanPause, CanResume, CanSeek, Seek
 	 */
-	bool CanScrub() const;
+	MEDIAUTILS_API bool CanScrub() const;
 
 	/**
 	 * Whether playback can jump to a position.
@@ -175,12 +175,12 @@ public:
 	 * @return true if seeking is allowed, false otherwise.
 	 * @see CanPause, CanResume, CanScrub, Seek
 	 */
-	bool CanSeek() const;
+	MEDIAUTILS_API bool CanSeek() const;
 
 	/**
 	 * Close the currently open media, if any.
 	 */
-	void Close();
+	MEDIAUTILS_API void Close();
 
 	/**
 	 * Get the number of channels in the specified audio track.
@@ -190,7 +190,7 @@ public:
 	 * @return Number of channels.
 	 * @see GetAudioTrackSampleRate, GetAudioTrackType
 	 */
-	uint32 GetAudioTrackChannels(int32 TrackIndex, int32 FormatIndex) const;
+	MEDIAUTILS_API uint32 GetAudioTrackChannels(int32 TrackIndex, int32 FormatIndex) const;
 
 	/**
 	 * Get the sample rate of the specified audio track.
@@ -200,7 +200,7 @@ public:
 	 * @return Samples per second.
 	 * @see GetAudioTrackChannels, GetAudioTrackType
 	 */
-	uint32 GetAudioTrackSampleRate(int32 TrackIndex, int32 FormatIndex) const;
+	MEDIAUTILS_API uint32 GetAudioTrackSampleRate(int32 TrackIndex, int32 FormatIndex) const;
 
 	/**
 	 * Get the type of the specified audio track format.
@@ -210,7 +210,7 @@ public:
 	 * @return Audio format type string.
 	 * @see GetAudioTrackSampleRate, GetAudioTrackSampleRate
 	 */
-	FString GetAudioTrackType(int32 TrackIndex, int32 FormatIndex) const;
+	MEDIAUTILS_API FString GetAudioTrackType(int32 TrackIndex, int32 FormatIndex) const;
 
 	/**
 	 * Get the media's duration.
@@ -218,7 +218,7 @@ public:
 	 * @return A time span representing the duration.
 	 * @see GetTime, Seek
 	 */
-	FTimespan GetDuration() const;
+	MEDIAUTILS_API FTimespan GetDuration() const;
 
 	/**
 	 * Get the player's globally unique identifier.
@@ -226,7 +226,7 @@ public:
 	 * @return The Guid.
 	 * @see SetGuid
 	 */
-	const FGuid& GetGuid();
+	MEDIAUTILS_API const FGuid& GetGuid();
 
 	/**
 	 * Get debug information about the player and currently opened media.
@@ -234,7 +234,7 @@ public:
 	 * @return Information string.
 	 * @see GetStats
 	 */
-	FString GetInfo() const;
+	MEDIAUTILS_API FString GetInfo() const;
 
 	/**
 	 * Get information about the media that is playing.
@@ -243,7 +243,7 @@ public:
 	 * @returns					Requested information, or empty if not available.
 	 * @see						UMediaPlayer::GetMediaInfo.
 	 */
-	FVariant GetMediaInfo(FName InfoName) const;
+	MEDIAUTILS_API FVariant GetMediaInfo(FName InfoName) const;
 
 	/**
 	 * Get the human readable name of the currently loaded media source.
@@ -251,7 +251,7 @@ public:
 	 * @return Media source name, or empty text if no media is opened
 	 * @see GetPlayerName, GetUrl
 	 */
-	FText GetMediaName() const;
+	MEDIAUTILS_API FText GetMediaName() const;
 
 	/**
 	 * Get meta data contained in the current stream
@@ -259,7 +259,7 @@ public:
 	 * @return Map with arrays of IMediaMetaDataItem entries describing any metadata found in the current stream
 	 * @note Listen to EMediaEvent::MetadataChanged to catch updates to this data
 	 */
-	TSharedPtr<TMap<FString, TArray<TUniquePtr<IMediaMetadataItem>>>, ESPMode::ThreadSafe> GetMediaMetadata() const;
+	MEDIAUTILS_API TSharedPtr<TMap<FString, TArray<TUniquePtr<IMediaMetadataItem>>>, ESPMode::ThreadSafe> GetMediaMetadata() const;
 
 	/**
 	 * Get the number of tracks of the given type.
@@ -268,7 +268,7 @@ public:
 	 * @return Number of tracks.
 	 * @see GetSelectedTrack, SelectTrack
 	 */
-	int32 GetNumTracks(EMediaTrackType TrackType) const;
+	MEDIAUTILS_API int32 GetNumTracks(EMediaTrackType TrackType) const;
 
 	/**
 	 * Get the number of formats of the specified track.
@@ -278,7 +278,7 @@ public:
 	 * @return Number of formats.
 	 * @see GetNumTracks, GetSelectedTrack, SelectTrack
 	 */
-	int32 GetNumTrackFormats(EMediaTrackType TrackType, int32 TrackIndex) const;
+	MEDIAUTILS_API int32 GetNumTrackFormats(EMediaTrackType TrackType, int32 TrackIndex) const;
 
 	/**
 	 * Get the low-level player associated with this object.
@@ -296,7 +296,7 @@ public:
 	 * @return Player name, or NAME_None if not available.
 	 * @see GetMediaName
 	 */
-	FName GetPlayerName() const;
+	MEDIAUTILS_API FName GetPlayerName() const;
 
 	/**
 	 * Get the media's current playback rate.
@@ -304,7 +304,7 @@ public:
 	 * @return The playback rate.
 	 * @see SetRate, SupportsRate
 	 */
-	float GetRate() const;
+	MEDIAUTILS_API float GetRate() const;
 
 	/**
 	 * Get the index of the currently selected track of the given type.
@@ -313,7 +313,7 @@ public:
 	 * @return The index of the selected track, or INDEX_NONE if no track is active.
 	 * @see GetNumTracks, SelectTrack
 	 */
-	int32 GetSelectedTrack(EMediaTrackType TrackType) const;
+	MEDIAUTILS_API int32 GetSelectedTrack(EMediaTrackType TrackType) const;
 
 	/**
 	 * Get playback statistics information.
@@ -321,7 +321,7 @@ public:
 	 * @return Information string.
 	 * @see GetInfo
 	 */
-	FString GetStats() const;
+	MEDIAUTILS_API FString GetStats() const;
 
 	/**
 	 * Get the supported playback rates.
@@ -330,7 +330,7 @@ public:
 	 * @return The ranges of supported rates.
 	 * @see SetRate, SupportsRate
 	 */
-	TRangeSet<float> GetSupportedRates(bool Unthinned = true) const;
+	MEDIAUTILS_API TRangeSet<float> GetSupportedRates(bool Unthinned = true) const;
 
 	/**
 	 * Get the media's current playback time.
@@ -338,14 +338,14 @@ public:
 	 * @return Playback time.
 	 * @see GetDuration, Seek
 	 */
-	FTimespan GetTime() const;
+	MEDIAUTILS_API FTimespan GetTime() const;
 
 	/**
 	 * Get the media's current playback time stamp.
 	 *
 	 * @return Playback time stamp.
 	 */
-	FMediaTimeStamp GetTimeStamp() const;
+	MEDIAUTILS_API FMediaTimeStamp GetTimeStamp() const;
 
 	/**
 	 * Get the media's current playback time stamp in a "display" version
@@ -355,7 +355,7 @@ public:
 	 * @note The timestamp returned here will reflect a user-logic oriented version.
 	 *       (e.g. during seeks this will return the seek target rather than the last valid frame still displayed)
 	 */
-	FMediaTimeStamp GetDisplayTimeStamp() const;
+	MEDIAUTILS_API FMediaTimeStamp GetDisplayTimeStamp() const;
 
 	/**
 	 * Get the human readable name of the specified track.
@@ -365,7 +365,7 @@ public:
 	 * @return Display name.
 	 * @see GetNumTracks, GetTrackLanguage
 	 */
-	FText GetTrackDisplayName(EMediaTrackType TrackType, int32 TrackIndex) const;
+	MEDIAUTILS_API FText GetTrackDisplayName(EMediaTrackType TrackType, int32 TrackIndex) const;
 
 	/**
 	 * Get the index of the active format of the specified track.
@@ -375,7 +375,7 @@ public:
 	 * @return The index of the selected format.
 	 * @see GetNumTrackFormats, GetSelectedTrack, SetTrackFormat
 	 */
-	int32 GetTrackFormat(EMediaTrackType TrackType, int32 TrackIndex) const;
+	MEDIAUTILS_API int32 GetTrackFormat(EMediaTrackType TrackType, int32 TrackIndex) const;
 
 	/**
 	 * Get the language tag of the specified track.
@@ -385,7 +385,7 @@ public:
 	 * @return Language tag, i.e. "en-US" for English, or "und" for undefined.
 	 * @see GetNumTracks, GetTrackDisplayName
 	 */
-	FString GetTrackLanguage(EMediaTrackType TrackType, int32 TrackIndex) const;
+	MEDIAUTILS_API FString GetTrackLanguage(EMediaTrackType TrackType, int32 TrackIndex) const;
 
 	/**
 	 * Get the URL of the currently loaded media, if any.
@@ -405,7 +405,7 @@ public:
 	 * @return Aspect ratio.
 	 * @see GetVideoTrackDimensions, GetVideoTrackFrameRate, GetVideoTrackFrameRates, GetVideoTrackType
 	 */
-	float GetVideoTrackAspectRatio(int32 TrackIndex, int32 FormatIndex) const;
+	MEDIAUTILS_API float GetVideoTrackAspectRatio(int32 TrackIndex, int32 FormatIndex) const;
 
 	/**
 	 * Get the width and height of the specified video track.
@@ -415,7 +415,7 @@ public:
 	 * @return Video dimensions.
 	 * @see GetVideoTrackAspectRatio, GetVideoTrackFrameRate, GetVideoTrackFrameRates, GetVideoTrackType
 	 */
-	FIntPoint GetVideoTrackDimensions(int32 TrackIndex, int32 FormatIndex) const;
+	MEDIAUTILS_API FIntPoint GetVideoTrackDimensions(int32 TrackIndex, int32 FormatIndex) const;
 
 	/**
 	 * Get frame rate of the specified video track.
@@ -425,7 +425,7 @@ public:
 	 * @return Video frame rate.
 	 * @see GetVideoTrackAspectRatio, GetVideoTrackDimensions, GetVideoTrackFrameRates, GetVideoTrackType
 	 */
-	float GetVideoTrackFrameRate(int32 TrackIndex, int32 FormatIndex) const;
+	MEDIAUTILS_API float GetVideoTrackFrameRate(int32 TrackIndex, int32 FormatIndex) const;
 
 	/**
 	 * Get the supported range of frame rates of the specified video track.
@@ -435,7 +435,7 @@ public:
 	 * @return Frame rate range (in frames per second).
 	 * @see GetVideoTrackAspectRatio, GetVideoTrackDimensions, GetVideoTrackFrameRate, GetVideoTrackType
 	 */
-	TRange<float> GetVideoTrackFrameRates(int32 TrackIndex, int32 FormatIndex) const;
+	MEDIAUTILS_API TRange<float> GetVideoTrackFrameRates(int32 TrackIndex, int32 FormatIndex) const;
 
 	/**
 	 * Get the type of the specified video track format.
@@ -445,7 +445,7 @@ public:
 	 * @return Video format type string.
 	 * @see GetVideoTrackAspectRatio, GetVideoTrackDimensions, GetVideoTrackFrameRate, GetVideoTrackFrameRates
 	 */
-	FString GetVideoTrackType(int32 TrackIndex, int32 FormatIndex) const;
+	MEDIAUTILS_API FString GetVideoTrackType(int32 TrackIndex, int32 FormatIndex) const;
 
 	/**
 	 * Get the field of view.
@@ -455,7 +455,7 @@ public:
 	 * @return true on success, false if feature is not available or if field of view has never been set.
 	 * @see GetViewOrientation, SetViewField
 	 */
-	bool GetViewField(float& OutHorizontal, float& OutVertical) const;
+	MEDIAUTILS_API bool GetViewField(float& OutHorizontal, float& OutVertical) const;
 
 	/**
 	 * Get the view's orientation.
@@ -464,14 +464,14 @@ public:
 	 * @return true on success, false if feature is not available or if orientation has never been set.
 	 * @see GetViewField, SetViewOrientation
 	 */
-	bool GetViewOrientation(FQuat& OutOrientation) const;
+	MEDIAUTILS_API bool GetViewOrientation(FQuat& OutOrientation) const;
 
 	/**
 	 * Check whether the player is in an error state.
 	 *
 	 * @see IsReady
 	 */
-	bool HasError() const;
+	MEDIAUTILS_API bool HasError() const;
 
 	/**
 	 * Whether the player is currently buffering data.
@@ -479,7 +479,7 @@ public:
 	 * @return true if buffering, false otherwise.
 	 * @see GetState, IsConnecting, IsLooping
 	 */
-	bool IsBuffering() const;
+	MEDIAUTILS_API bool IsBuffering() const;
 
 	/**
 	 * Whether the player is currently connecting to a media source.
@@ -487,7 +487,7 @@ public:
 	 * @return true if connecting, false otherwise.
 	 * @see GetState, IsBuffering, IsLooping
 	 */
-	bool IsConnecting() const;
+	MEDIAUTILS_API bool IsConnecting() const;
 
 	/**
 	 * Whether playback is looping.
@@ -495,7 +495,7 @@ public:
 	 * @return true if looping, false otherwise.
 	 * @see GetState, IsBuffering, IsConnecting, SetLooping
 	 */
-	bool IsLooping() const;
+	MEDIAUTILS_API bool IsLooping() const;
 
 	/**
 	 * Whether playback is currently paused.
@@ -503,7 +503,7 @@ public:
 	 * @return true if playback is paused, false otherwise.
 	 * @see CanPause, IsPlaying, IsReady, Pause
 	 */
-	bool IsPaused() const;
+	MEDIAUTILS_API bool IsPaused() const;
 
 	/**
 	 * Whether playback is in progress.
@@ -511,7 +511,7 @@ public:
 	 * @return true if playback has started, false otherwise.
 	 * @see CanPlay, IsPaused, IsReady, Play
 	 */
-	bool IsPlaying() const;
+	MEDIAUTILS_API bool IsPlaying() const;
 
 	/**
 	 * Whether the media is currently opening or buffering.
@@ -519,14 +519,14 @@ public:
 	 * @return true if playback is being prepared, false otherwise.
 	 * @see CanPlay, IsPaused, IsReady, Play
 	 */
-	bool IsPreparing() const;
+	MEDIAUTILS_API bool IsPreparing() const;
 
 	/**
 	 * Whether media is currently closed.
 	 *
 	 * @return true if media is closed, false otherwise.
 	 */
-	bool IsClosed() const;
+	MEDIAUTILS_API bool IsClosed() const;
 
 	/**
 	 * Whether media is ready for playback.
@@ -537,7 +537,7 @@ public:
 	 * @return true if media is ready, false otherwise.
 	 * @see HasError, IsPaused, IsPlaying, Stop
 	 */
-	bool IsReady() const;
+	MEDIAUTILS_API bool IsReady() const;
 
 	/**
 	 * Open a media source from a URL with optional parameters.
@@ -547,7 +547,7 @@ public:
 	 * @param PlayerOptions Optional player parameters.
 	 * @return true if the media is being opened, false otherwise.
 	 */
-	bool Open(const FString& Url, const IMediaOptions* Options, const FMediaPlayerOptions* PlayerOptions = nullptr);
+	MEDIAUTILS_API bool Open(const FString& Url, const IMediaOptions* Options, const FMediaPlayerOptions* PlayerOptions = nullptr);
 
 	/**
 	 * Query the time ranges of cached media samples for the specified caching state.
@@ -555,7 +555,7 @@ public:
 	 * @param State The sample state we're interested in.
 	 * @param OutTimeRanges Will contain the set of matching sample time ranges.
 	 */
-	void QueryCacheState(EMediaTrackType TrackType, EMediaCacheState State, TRangeSet<FTimespan>& OutTimeRanges) const;
+	MEDIAUTILS_API void QueryCacheState(EMediaTrackType TrackType, EMediaCacheState State, TRangeSet<FTimespan>& OutTimeRanges) const;
 
 	/**
 	 * Seeks to the specified playback time.
@@ -564,7 +564,7 @@ public:
 	 * @return true on success, false otherwise.
 	 * @see GetTime, Rewind
 	 */
-	bool Seek(const FTimespan& Time);
+	MEDIAUTILS_API bool Seek(const FTimespan& Time);
 
 	/**
 	 * Select the active track of the given type.
@@ -578,7 +578,7 @@ public:
 	 * @return true if the track was selected, false otherwise.
 	 * @see GetNumTracks, GetSelectedTrack, SetTrackFormat
 	 */
-	bool SelectTrack(EMediaTrackType TrackType, int32 TrackIndex);
+	MEDIAUTILS_API bool SelectTrack(EMediaTrackType TrackType, int32 TrackIndex);
 
 	/**
 	 * Set the time on which to block.
@@ -590,7 +590,7 @@ public:
 	 * @see TickFetch
 	 * @note Deprecated: Use SetBlockOnTimeRange instead
 	 */
-	void SetBlockOnTime(const FTimespan& Time);
+	MEDIAUTILS_API void SetBlockOnTime(const FTimespan& Time);
 
 	/**
 	 * Set the time range on which to block.
@@ -600,7 +600,7 @@ public:
 	 *
 	 * @param TimeRange The time range to block on, use empty range to disable
 	 */
-	void SetBlockOnTimeRange(const TRange<FTimespan>& TimeRange);
+	MEDIAUTILS_API void SetBlockOnTimeRange(const TRange<FTimespan>& TimeRange);
 
 	/**
 	 * Set sample caching options.
@@ -608,7 +608,7 @@ public:
 	 * @param Ahead Duration of samples to cache ahead of the play head.
 	 * @param Behind Duration of samples to cache behind the play head.
 	 */
-	void SetCacheWindow(FTimespan Ahead, FTimespan Behind);
+	MEDIAUTILS_API void SetCacheWindow(FTimespan Ahead, FTimespan Behind);
 
 	/**
 	 * Set the player's globally unique identifier.
@@ -616,7 +616,7 @@ public:
 	 * @param Guid The GUID to set.
 	 * @see GetGuid
 	 */
-	void SetGuid(FGuid& Guid);
+	MEDIAUTILS_API void SetGuid(FGuid& Guid);
 
 	/**
 	 * Enables or disables playback looping.
@@ -625,12 +625,12 @@ public:
 	 * @return true on success, false otherwise.
 	 * @see IsLooping
 	 */
-	bool SetLooping(bool Looping);
+	MEDIAUTILS_API bool SetLooping(bool Looping);
 
 	/**
 	 * Changes media ooptions on the player.
 	 */
-	void SetMediaOptions(const IMediaOptions* Options);
+	MEDIAUTILS_API void SetMediaOptions(const IMediaOptions* Options);
 
 	/**
 	 * Changes the media's playback rate.
@@ -639,7 +639,7 @@ public:
 	 * @return true on success, false otherwise.
 	 * @see GetRate, SupportsRate
 	 */
-	bool SetRate(float Rate);
+	MEDIAUTILS_API bool SetRate(float Rate);
 
 	/**
 	 * Changes the media's native volume.
@@ -648,7 +648,7 @@ public:
 	 * @return true on success, false otherwise.
 	 * @see NativeAudioOut
 	 */
-	bool SetNativeVolume(float Volume);
+	MEDIAUTILS_API bool SetNativeVolume(float Volume);
 
 	/**
 	 * Set the format on the specified track.
@@ -663,7 +663,7 @@ public:
 	 * @return true if the track was selected, false otherwise.
 	 * @see GetNumTrackFormats, GetNumTracks, GetTrackFormat, SelectTrack
 	 */
-	bool SetTrackFormat(EMediaTrackType TrackType, int32 TrackIndex, int32 FormatIndex);
+	MEDIAUTILS_API bool SetTrackFormat(EMediaTrackType TrackType, int32 TrackIndex, int32 FormatIndex);
 
 	/**
 	 * Set the frame rate of the specified video track.
@@ -674,7 +674,7 @@ public:
 	 * @return true on success, false otherwise.
 	 * @see GetVideoTrackAspectRatio, GetVideoTrackDimensions, GetVideoTrackFrameRate, GetVideoTrackFrameRates, GetVideoTrackType
 	 */
-	bool SetVideoTrackFrameRate(int32 TrackIndex, int32 FormatIndex, float FrameRate);
+	MEDIAUTILS_API bool SetVideoTrackFrameRate(int32 TrackIndex, int32 FormatIndex, float FrameRate);
 
 	/**
 	 * Set the field of view.
@@ -685,7 +685,7 @@ public:
 	 * @return true on success, false otherwise.
 	 * @see GetViewField, SetViewOrientation
 	 */
-	bool SetViewField(float Horizontal, float Vertical, bool Absolute);
+	MEDIAUTILS_API bool SetViewField(float Horizontal, float Vertical, bool Absolute);
 
 	/**
 	 * Set the view's orientation.
@@ -695,7 +695,7 @@ public:
 	 * @return true on success, false otherwise.
 	 * @see GetViewOrientation, SetViewField
 	 */
-	bool SetViewOrientation(const FQuat& Orientation, bool Absolute);
+	MEDIAUTILS_API bool SetViewOrientation(const FQuat& Orientation, bool Absolute);
 
 	/**
 	 * Whether the specified playback rate is supported.
@@ -704,7 +704,7 @@ public:
 	 * @param Unthinned Whether no frames should be dropped at the given rate.
 	 * @see CanScrub, CanSeek
 	 */
-	bool SupportsRate(float Rate, bool Unthinned) const;
+	MEDIAUTILS_API bool SupportsRate(float Rate, bool Unthinned) const;
 
 	/**
 	 * Record last audio sample played to track audio sync (for automated tests)
@@ -712,21 +712,21 @@ public:
 	 * @param SampleTime Time of media sample currently being played
 	 * @return true if playback is being prepared, false otherwise.
 	 */
-	void SetLastAudioRenderedSampleTime(FTimespan SampleTime);
+	MEDIAUTILS_API void SetLastAudioRenderedSampleTime(FTimespan SampleTime);
 
 	/**
 	 * Get time of last audio sample played
 	 *
 	 * @return Time of last audio sample played.
 	 */
-	FTimespan GetLastAudioRenderedSampleTime() const;
+	MEDIAUTILS_API FTimespan GetLastAudioRenderedSampleTime() const;
 
 	/**
 	 * Sets whether the player can broadcast events when running on a thread other than the game thread.
 	 *
 	 * @param bInAreEventsSafeForAnyThread If true then allow broadcast when not on the game thread.
 	 */
-	void SetAreEventsSafeForAnyThread(bool bInAreEventsSafeForAnyThread);
+	MEDIAUTILS_API void SetAreEventsSafeForAnyThread(bool bInAreEventsSafeForAnyThread);
 
 public:
 
@@ -741,15 +741,15 @@ public:
 
 	//~ IMediaClockSink interface
 
-	virtual void TickFetch(FTimespan DeltaTime, FTimespan Timecode) override;
-	virtual void TickInput(FTimespan DeltaTime, FTimespan Timecode) override;
-	virtual void TickOutput(FTimespan DeltaTime, FTimespan Timecode) override;
+	MEDIAUTILS_API virtual void TickFetch(FTimespan DeltaTime, FTimespan Timecode) override;
+	MEDIAUTILS_API virtual void TickInput(FTimespan DeltaTime, FTimespan Timecode) override;
+	MEDIAUTILS_API virtual void TickOutput(FTimespan DeltaTime, FTimespan Timecode) override;
 
 public:
 
 	//~ IMediaTickable interface
 
-	virtual void TickTickable() override;
+	MEDIAUTILS_API virtual void TickTickable() override;
 
 protected:
 
@@ -758,13 +758,13 @@ protected:
 	 *
 	 * @return true if sample fetching should block, false otherwise.
 	 */
-	bool BlockOnFetch() const;
+	MEDIAUTILS_API bool BlockOnFetch() const;
 
 	/** Flush all media sample sinks & player plugin. */
-	void Flush(bool bExcludePlayer = false, bool bOnSeek = false);
+	MEDIAUTILS_API void Flush(bool bExcludePlayer = false, bool bOnSeek = false);
 
 	/** Internal function to retrieve the current timestamp */
-	FMediaTimeStamp GetTimeStampInternal(bool bForDisplay) const;
+	MEDIAUTILS_API FMediaTimeStamp GetTimeStampInternal(bool bForDisplay) const;
 
 	/**
 	 * Get details about the specified audio track format.
@@ -775,7 +775,7 @@ protected:
 	 * @return true on success, false otherwise.
 	 * @see GetVideoTrackFormat
 	 */
-	bool GetAudioTrackFormat(int32 TrackIndex, int32 FormatIndex, FMediaAudioTrackFormat& OutFormat) const;
+	MEDIAUTILS_API bool GetAudioTrackFormat(int32 TrackIndex, int32 FormatIndex, FMediaAudioTrackFormat& OutFormat) const;
 
 	/**
 	 * Get a player that can play the specified media URL.
@@ -784,7 +784,7 @@ protected:
 	 * @param Options The media options for the URL.
 	 * @return The player if found, or nullptr otherwise.
 	 */
-	IMediaPlayerFactory *GetPlayerFactoryForUrl(const FString& Url, const IMediaOptions* Options) const;
+	MEDIAUTILS_API IMediaPlayerFactory *GetPlayerFactoryForUrl(const FString& Url, const IMediaOptions* Options) const;
 
 	/**
 	 * Get details about the specified audio track format.
@@ -795,7 +795,7 @@ protected:
 	 * @return true on success, false otherwise.
 	 * @see GetVideoTrackFormat
 	 */
-	bool GetVideoTrackFormat(int32 TrackIndex, int32 FormatIndex, FMediaVideoTrackFormat& OutFormat) const;
+	MEDIAUTILS_API bool GetVideoTrackFormat(int32 TrackIndex, int32 FormatIndex, FMediaVideoTrackFormat& OutFormat) const;
 
 	/**
 	 * Process the given media event.
@@ -803,7 +803,7 @@ protected:
 	 * @param Event The event to process.
 	 * @param bIsBroadcastAllowed If true then we can broadcast events, if false then they will sent when possible.
 	 **/
-	void ProcessEvent(EMediaEvent Event, bool bIsBroadcastAllowed);
+	MEDIAUTILS_API void ProcessEvent(EMediaEvent Event, bool bIsBroadcastAllowed);
 
 private:
 	struct FTrackSelection
@@ -813,61 +813,61 @@ private:
 	} TrackSelection;
 
 	/** Reset all tracks. */
-	void ResetTracks();
+	MEDIAUTILS_API void ResetTracks();
 
 	/** Setup track selection with player. */
-	void UpdateTrackSelectionWithPlayer();
+	MEDIAUTILS_API void UpdateTrackSelectionWithPlayer();
 
 	/** Select the default media tracks. */
-	void SelectDefaultTracks();
+	MEDIAUTILS_API void SelectDefaultTracks();
 
 protected:
-	bool HaveAudioPlayback() const;
-	bool HaveVideoPlayback() const;
-	float GetUnpausedRate() const;
+	MEDIAUTILS_API bool HaveAudioPlayback() const;
+	MEDIAUTILS_API bool HaveVideoPlayback() const;
+	MEDIAUTILS_API float GetUnpausedRate() const;
 
 
 protected:
 
 	//~ IMediaEventSink interface
 
-	void ReceiveMediaEvent(EMediaEvent Event) override;
+	MEDIAUTILS_API void ReceiveMediaEvent(EMediaEvent Event) override;
 
 private:
 	friend class FMediaPlayerLifecycleManagerDelegateControl;
 
 	// Internal
-	bool NotifyLifetimeManagerDelegate_PlayerOpen(IMediaPlayerLifecycleManagerDelegate::IControlRef& NewLifecycleManagerDelegateControl, const FString& InUrl, const IMediaOptions* Options, const FMediaPlayerOptions* InPlayerOptions, IMediaPlayerFactory* PlayerFactory, bool bWillCreatePlayer, uint32 WillUseNewResources, uint64 NewPlayerInstanceID);
-	bool NotifyLifetimeManagerDelegate_PlayerCreated();
-	bool NotifyLifetimeManagerDelegate_PlayerCreateFailed();
-	bool NotifyLifetimeManagerDelegate_PlayerClosed();
-	bool NotifyLifetimeManagerDelegate_PlayerDestroyed();
-	bool NotifyLifetimeManagerDelegate_PlayerResourcesReleased(uint32 ResourceFlags);
-
-	void ProcessAudioSamples(IMediaSamples& Samples, const TRange<FMediaTimeStamp>& TimeRange);
-	bool ProcessVideoSamples(IMediaSamples& Samples, const TRange<FMediaTimeStamp>& TimeRange);
-	void ProcessCaptionSamples(IMediaSamples& Samples, const TRange<FMediaTimeStamp>& TimeRange);
-	void ProcessSubtitleSamples(IMediaSamples& Samples, const TRange<FMediaTimeStamp>& TimeRange);
-	void ProcessMetadataSamples(IMediaSamples& Samples, const TRange<FMediaTimeStamp>& TimeRange);
-
-	void ProcessAudioSamplesV1(IMediaSamples& Samples, TRange<FTimespan> TimeRange);
-	void ProcessVideoSamplesV1(IMediaSamples& Samples, TRange<FTimespan> TimeRange);
-	void ProcessSubtitleSamplesV1(IMediaSamples& Samples, TRange<FTimespan> TimeRange);
-	void ProcessCaptionSamplesV1(IMediaSamples& Samples, TRange<FTimespan> TimeRange);
-	void ProcessMetadataSamplesV1(IMediaSamples& Samples, TRange<FTimespan> TimeRange);
-
-	bool IsVideoSampleStillGood(const TRange<FMediaTimeStamp>& LastSampleTimeRange, const TRange<FMediaTimeStamp>& TimeRange, bool bReverse) const;
-	void MonitorAudioEnablement();
-	void UpdateSeekStatus(const FMediaTimeStamp* pCheckTimeStamp = nullptr);
-	void PreSampleProcessingTimeHandling();
-	bool GetCurrentPlaybackTimeRange(TRange<FMediaTimeStamp>& TimeRange, float Rate, FTimespan DeltaTime, bool bPurgeSampleRelated) const;
-	void PostSampleProcessingTimeHandling(FTimespan DeltaTime);
-
-	void DestroyPlayer();
-
-	bool ContinueOpen(IMediaPlayerLifecycleManagerDelegate::IControlRef NewLifecycleManagerDelegateControl, const FString& Url, const IMediaOptions* Options, const FMediaPlayerOptions* PlayerOptions, IMediaPlayerFactory* PlayerFactory, TSharedPtr<IMediaPlayer, ESPMode::ThreadSafe> ReusedPlayer, bool bCreateNewPlayer, uint64 NewPlayerInstanceID);
-
-	void SendSinkEvent(EMediaSampleSinkEvent Event, const FMediaSampleSinkEventData& Data);
+	MEDIAUTILS_API bool NotifyLifetimeManagerDelegate_PlayerOpen(IMediaPlayerLifecycleManagerDelegate::IControlRef& NewLifecycleManagerDelegateControl, const FString& InUrl, const IMediaOptions* Options, const FMediaPlayerOptions* InPlayerOptions, IMediaPlayerFactory* PlayerFactory, bool bWillCreatePlayer, uint32 WillUseNewResources, uint64 NewPlayerInstanceID);
+	MEDIAUTILS_API bool NotifyLifetimeManagerDelegate_PlayerCreated();
+	MEDIAUTILS_API bool NotifyLifetimeManagerDelegate_PlayerCreateFailed();
+	MEDIAUTILS_API bool NotifyLifetimeManagerDelegate_PlayerClosed();
+	MEDIAUTILS_API bool NotifyLifetimeManagerDelegate_PlayerDestroyed();
+	MEDIAUTILS_API bool NotifyLifetimeManagerDelegate_PlayerResourcesReleased(uint32 ResourceFlags);
+
+	MEDIAUTILS_API void ProcessAudioSamples(IMediaSamples& Samples, const TRange<FMediaTimeStamp>& TimeRange);
+	MEDIAUTILS_API bool ProcessVideoSamples(IMediaSamples& Samples, const TRange<FMediaTimeStamp>& TimeRange);
+	MEDIAUTILS_API void ProcessCaptionSamples(IMediaSamples& Samples, const TRange<FMediaTimeStamp>& TimeRange);
+	MEDIAUTILS_API void ProcessSubtitleSamples(IMediaSamples& Samples, const TRange<FMediaTimeStamp>& TimeRange);
+	MEDIAUTILS_API void ProcessMetadataSamples(IMediaSamples& Samples, const TRange<FMediaTimeStamp>& TimeRange);
+
+	MEDIAUTILS_API void ProcessAudioSamplesV1(IMediaSamples& Samples, TRange<FTimespan> TimeRange);
+	MEDIAUTILS_API void ProcessVideoSamplesV1(IMediaSamples& Samples, TRange<FTimespan> TimeRange);
+	MEDIAUTILS_API void ProcessSubtitleSamplesV1(IMediaSamples& Samples, TRange<FTimespan> TimeRange);
+	MEDIAUTILS_API void ProcessCaptionSamplesV1(IMediaSamples& Samples, TRange<FTimespan> TimeRange);
+	MEDIAUTILS_API void ProcessMetadataSamplesV1(IMediaSamples& Samples, TRange<FTimespan> TimeRange);
+
+	MEDIAUTILS_API bool IsVideoSampleStillGood(const TRange<FMediaTimeStamp>& LastSampleTimeRange, const TRange<FMediaTimeStamp>& TimeRange, bool bReverse) const;
+	MEDIAUTILS_API void MonitorAudioEnablement();
+	MEDIAUTILS_API void UpdateSeekStatus(const FMediaTimeStamp* pCheckTimeStamp = nullptr);
+	MEDIAUTILS_API void PreSampleProcessingTimeHandling();
+	MEDIAUTILS_API bool GetCurrentPlaybackTimeRange(TRange<FMediaTimeStamp>& TimeRange, float Rate, FTimespan DeltaTime, bool bPurgeSampleRelated) const;
+	MEDIAUTILS_API void PostSampleProcessingTimeHandling(FTimespan DeltaTime);
+
+	MEDIAUTILS_API void DestroyPlayer();
+
+	MEDIAUTILS_API bool ContinueOpen(IMediaPlayerLifecycleManagerDelegate::IControlRef NewLifecycleManagerDelegateControl, const FString& Url, const IMediaOptions* Options, const FMediaPlayerOptions* PlayerOptions, IMediaPlayerFactory* PlayerFactory, TSharedPtr<IMediaPlayer, ESPMode::ThreadSafe> ReusedPlayer, bool bCreateNewPlayer, uint64 NewPlayerInstanceID);
+
+	MEDIAUTILS_API void SendSinkEvent(EMediaSampleSinkEvent Event, const FMediaSampleSinkEventData& Data);
 
 	/** Audio sample sinks. */
 	TWeakPtr<FMediaAudioSampleSink, ESPMode::ThreadSafe> PrimaryAudioSink;
@@ -886,7 +886,7 @@ private:
 	FMediaVideoSampleSinks VideoSampleSinks;
 
 private:
-	TRange<FMediaTimeStamp> GetAdjustedBlockOnRange() const;
+	MEDIAUTILS_API TRange<FMediaTimeStamp> GetAdjustedBlockOnRange() const;
 
 	class FBlockOnRange
 	{
diff --git a/Engine/Source/Runtime/MediaUtils/Public/MediaRecorder.h b/Engine/Source/Runtime/MediaUtils/Public/MediaRecorder.h
--- a/Engine/Source/Runtime/MediaUtils/Public/MediaRecorder.h
+++ b/Engine/Source/Runtime/MediaUtils/Public/MediaRecorder.h
@@ -20,12 +20,12 @@ class IImageWriteQueue;
  * Loop, seek or reverse and not supported.
  * Currently only records texture samples and that support sample format 8bit BGRA, half float RBGA
  */
-class MEDIAUTILS_API FMediaRecorder
+class FMediaRecorder
 {
 public:
 
 	/** Default constructor. */
-	FMediaRecorder();
+	MEDIAUTILS_API FMediaRecorder();
 
 public:
 
@@ -70,12 +70,12 @@ public:
 	/**
 	 * Start recording samples from a given media player.
 	 */
-	void StartRecording(const FMediaRecorderData& InRecoderData);
+	MEDIAUTILS_API void StartRecording(const FMediaRecorderData& InRecoderData);
 
 	/**
 	 * Stop recording media samples.
 	 */
-	void StopRecording();
+	MEDIAUTILS_API void StopRecording();
 
 	/**
 	 * Is currently recording media samples.
@@ -85,7 +85,7 @@ public:
 	/**
 	 * Blocking call that will wait for all frames to be recorded before returning.
 	 */
-	bool WaitPendingTasks(const FTimespan& InDuration);
+	MEDIAUTILS_API bool WaitPendingTasks(const FTimespan& InDuration);
 
 protected:
 
@@ -94,7 +94,7 @@ protected:
 	 *
 	 * @param Timecode The current timecode.
 	 */
-	void TickRecording();
+	MEDIAUTILS_API void TickRecording();
 
 private:
 
diff --git a/Engine/Source/Runtime/MediaUtils/Public/MediaSampleCache.h b/Engine/Source/Runtime/MediaUtils/Public/MediaSampleCache.h
--- a/Engine/Source/Runtime/MediaUtils/Public/MediaSampleCache.h
+++ b/Engine/Source/Runtime/MediaUtils/Public/MediaSampleCache.h
@@ -27,7 +27,7 @@ template <typename ElementType> class TRangeSet;
 /**
  * Implements a cache for media samples.
  */
-class MEDIAUTILS_API FMediaSampleCache
+class FMediaSampleCache
 {
 	/** Key functions template for the sample cache sets. */
 	template<typename SampleType>
@@ -57,7 +57,7 @@ class MEDIAUTILS_API FMediaSampleCache
 public:
 
 	/** Default constructor. */
-	FMediaSampleCache();
+	MEDIAUTILS_API FMediaSampleCache();
 
 public:
 
@@ -66,7 +66,7 @@ public:
 	 *
 	 * @see Update
 	 */
-	void Empty();
+	MEDIAUTILS_API void Empty();
 
 	/**
 	 * Get the audio sample for the specified play time.
@@ -77,7 +77,7 @@ public:
 	 * @return The sample, or nullptr if no sample available.
 	 * @see GetAudioSamples, GetCachedVideoSampleRanges, GetOverlaySamples
 	 */
-	TSharedPtr<IMediaAudioSample, ESPMode::ThreadSafe> GetAudioSample(FTimespan Time);
+	MEDIAUTILS_API TSharedPtr<IMediaAudioSample, ESPMode::ThreadSafe> GetAudioSample(FTimespan Time);
 
 	/**
 	 * Get the time ranges of audio samples currently in the cache.
@@ -85,7 +85,7 @@ public:
 	 * @param OutTimeRanges Will contain the set of cached sample time ranges.
 	 * @see GetAudioSample, GetCachedVideoSampleRanges
 	 */
-	void GetCachedAudioSampleRanges(TRangeSet<FTimespan>& OutTimeRanges) const;
+	MEDIAUTILS_API void GetCachedAudioSampleRanges(TRangeSet<FTimespan>& OutTimeRanges) const;
 
 	/**
 	 * Get the time ranges of video samples currently in the cache.
@@ -93,7 +93,7 @@ public:
 	 * @param OutTimeRanges Will contain the set of cached sample time ranges.
 	 * @see GetCachedAudioSampleRanges, GetVideoSample
 	 */
-	void GetCachedVideoSampleRanges(TRangeSet<FTimespan>& OutTimeRanges) const;
+	MEDIAUTILS_API void GetCachedVideoSampleRanges(TRangeSet<FTimespan>& OutTimeRanges) const;
 
 	/**
 	 * Get the text overlay samples for the specified time.
@@ -102,7 +102,7 @@ public:
 	 * @param OutSamples Will contain the overlay samples.
 	 * @see GetAudioSample, GetVideoSample
 	 */
-	void GetOverlaySamples(FTimespan Time, TArray<TSharedPtr<IMediaOverlaySample, ESPMode::ThreadSafe>>& OutSamples);
+	MEDIAUTILS_API void GetOverlaySamples(FTimespan Time, TArray<TSharedPtr<IMediaOverlaySample, ESPMode::ThreadSafe>>& OutSamples);
 
 	/**
 	 * Get the video sample for the specified play time.
@@ -116,7 +116,7 @@ public:
 	 * @return The sample, or nullptr if no sample available.
 	 * @see GetAudioSample, GetCachedVideoSampleRanges, GetOverlaySamples
 	 */
-	TSharedPtr<IMediaTextureSample, ESPMode::ThreadSafe> GetVideoSample(FTimespan Time, bool Forward);
+	MEDIAUTILS_API TSharedPtr<IMediaTextureSample, ESPMode::ThreadSafe> GetVideoSample(FTimespan Time, bool Forward);
 
 	/**
 	 * Set the time window of samples to cache.
@@ -141,7 +141,7 @@ public:
 	 * @param Time The current play time.
 	 * @see Initialize, Shutdown
 	 */
-	void Tick(FTimespan DeltaTime, float Rate, FTimespan Time);
+	MEDIAUTILS_API void Tick(FTimespan DeltaTime, float Rate, FTimespan Time);
 
 private:
 
diff --git a/Engine/Source/Runtime/MediaUtils/Public/MediaSamples.h b/Engine/Source/Runtime/MediaUtils/Public/MediaSamples.h
--- a/Engine/Source/Runtime/MediaUtils/Public/MediaSamples.h
+++ b/Engine/Source/Runtime/MediaUtils/Public/MediaSamples.h
@@ -20,7 +20,7 @@ template <typename ElementType> class TRange;
 /**
  * General purpose media sample queue.
  */
-class MEDIAUTILS_API FMediaSamples
+class FMediaSamples
 	: public IMediaSamples
 {
 public:
@@ -30,7 +30,7 @@ public:
 	static constexpr uint32 MaxNumberOfQueuedSubtitlesSamples = 4;
 	static constexpr uint32 MaxNumberOfQueuedMetaDataSamples = 4;
 
-	FMediaSamples(uint32 InMaxNumberOfQueuedAudioSamples = MaxNumberOfQueuedAudioSamples,
+	MEDIAUTILS_API FMediaSamples(uint32 InMaxNumberOfQueuedAudioSamples = MaxNumberOfQueuedAudioSamples,
 				  uint32 InMaxNumberOfQueuedVideoSamples = MaxNumberOfQueuedVideoSamples,
 				  uint32 InMaxNumberOfQueuedCaptionSamples = MaxNumberOfQueuedCaptionSamples,
 				  uint32 InMaxNumberOfQueuedSubtitlesSamples = MaxNumberOfQueuedSubtitlesSamples,
@@ -150,20 +150,20 @@ public:
 
 	//~ IMediaSamples interface
 
-	virtual bool FetchAudio(TRange<FTimespan> TimeRange, TSharedPtr<IMediaAudioSample, ESPMode::ThreadSafe>& OutSample) override;
-	virtual bool FetchCaption(TRange<FTimespan> TimeRange, TSharedPtr<IMediaOverlaySample, ESPMode::ThreadSafe>& OutSample) override;
-	virtual bool FetchMetadata(TRange<FTimespan> TimeRange, TSharedPtr<IMediaBinarySample, ESPMode::ThreadSafe>& OutSample) override;
-	virtual bool FetchSubtitle(TRange<FTimespan> TimeRange, TSharedPtr<IMediaOverlaySample, ESPMode::ThreadSafe>& OutSample) override;
-	virtual bool FetchVideo(TRange<FTimespan> TimeRange, TSharedPtr<IMediaTextureSample, ESPMode::ThreadSafe>& OutSample) override;
-	virtual void FlushSamples() override;
+	MEDIAUTILS_API virtual bool FetchAudio(TRange<FTimespan> TimeRange, TSharedPtr<IMediaAudioSample, ESPMode::ThreadSafe>& OutSample) override;
+	MEDIAUTILS_API virtual bool FetchCaption(TRange<FTimespan> TimeRange, TSharedPtr<IMediaOverlaySample, ESPMode::ThreadSafe>& OutSample) override;
+	MEDIAUTILS_API virtual bool FetchMetadata(TRange<FTimespan> TimeRange, TSharedPtr<IMediaBinarySample, ESPMode::ThreadSafe>& OutSample) override;
+	MEDIAUTILS_API virtual bool FetchSubtitle(TRange<FTimespan> TimeRange, TSharedPtr<IMediaOverlaySample, ESPMode::ThreadSafe>& OutSample) override;
+	MEDIAUTILS_API virtual bool FetchVideo(TRange<FTimespan> TimeRange, TSharedPtr<IMediaTextureSample, ESPMode::ThreadSafe>& OutSample) override;
+	MEDIAUTILS_API virtual void FlushSamples() override;
 
-	virtual bool FetchAudio(TRange<FMediaTimeStamp> TimeRange, TSharedPtr<IMediaAudioSample, ESPMode::ThreadSafe>& OutSample) override;
-	virtual bool FetchCaption(TRange<FMediaTimeStamp> TimeRange, TSharedPtr<IMediaOverlaySample, ESPMode::ThreadSafe>& OutSample) override;
-	virtual bool FetchMetadata(TRange<FMediaTimeStamp> TimeRange, TSharedPtr<IMediaBinarySample, ESPMode::ThreadSafe>& OutSample) override;
-	virtual bool FetchSubtitle(TRange<FMediaTimeStamp> TimeRange, TSharedPtr<IMediaOverlaySample, ESPMode::ThreadSafe>& OutSample) override;
-	virtual bool FetchVideo(TRange<FMediaTimeStamp> TimeRange, TSharedPtr<IMediaTextureSample, ESPMode::ThreadSafe>& OutSample) override;
+	MEDIAUTILS_API virtual bool FetchAudio(TRange<FMediaTimeStamp> TimeRange, TSharedPtr<IMediaAudioSample, ESPMode::ThreadSafe>& OutSample) override;
+	MEDIAUTILS_API virtual bool FetchCaption(TRange<FMediaTimeStamp> TimeRange, TSharedPtr<IMediaOverlaySample, ESPMode::ThreadSafe>& OutSample) override;
+	MEDIAUTILS_API virtual bool FetchMetadata(TRange<FMediaTimeStamp> TimeRange, TSharedPtr<IMediaBinarySample, ESPMode::ThreadSafe>& OutSample) override;
+	MEDIAUTILS_API virtual bool FetchSubtitle(TRange<FMediaTimeStamp> TimeRange, TSharedPtr<IMediaOverlaySample, ESPMode::ThreadSafe>& OutSample) override;
+	MEDIAUTILS_API virtual bool FetchVideo(TRange<FMediaTimeStamp> TimeRange, TSharedPtr<IMediaTextureSample, ESPMode::ThreadSafe>& OutSample) override;
 
-	virtual EFetchBestSampleResult FetchBestVideoSampleForTimeRange(const TRange<FMediaTimeStamp> & TimeRange, TSharedPtr<IMediaTextureSample, ESPMode::ThreadSafe>& OutSample, bool bReverse) override;
+	MEDIAUTILS_API virtual EFetchBestSampleResult FetchBestVideoSampleForTimeRange(const TRange<FMediaTimeStamp> & TimeRange, TSharedPtr<IMediaTextureSample, ESPMode::ThreadSafe>& OutSample, bool bReverse) override;
 
 	/**
 	 * Peek next video sample's timestamp
@@ -205,16 +205,16 @@ public:
 		return MetadataSampleQueue.Discard(TimeRange, bReverse);
 	}
 
-	virtual uint32 PurgeOutdatedVideoSamples(const FMediaTimeStamp & ReferenceTime, bool bReversed, FTimespan MaxAge) override;
-	virtual uint32 PurgeOutdatedSubtitleSamples(const FMediaTimeStamp & ReferenceTime, bool bReversed, FTimespan MaxAge) override;
-	virtual uint32 PurgeOutdatedCaptionSamples(const FMediaTimeStamp& ReferenceTime, bool bReversed, FTimespan MaxAge) override;
-	virtual uint32 PurgeOutdatedMetadataSamples(const FMediaTimeStamp& ReferenceTime, bool bReversed, FTimespan MaxAge) override;
+	MEDIAUTILS_API virtual uint32 PurgeOutdatedVideoSamples(const FMediaTimeStamp & ReferenceTime, bool bReversed, FTimespan MaxAge) override;
+	MEDIAUTILS_API virtual uint32 PurgeOutdatedSubtitleSamples(const FMediaTimeStamp & ReferenceTime, bool bReversed, FTimespan MaxAge) override;
+	MEDIAUTILS_API virtual uint32 PurgeOutdatedCaptionSamples(const FMediaTimeStamp& ReferenceTime, bool bReversed, FTimespan MaxAge) override;
+	MEDIAUTILS_API virtual uint32 PurgeOutdatedMetadataSamples(const FMediaTimeStamp& ReferenceTime, bool bReversed, FTimespan MaxAge) override;
 
-	virtual bool CanReceiveVideoSamples(uint32 Num) const override;
-	virtual bool CanReceiveAudioSamples(uint32 Num) const override;
-	virtual bool CanReceiveSubtitleSamples(uint32 Num) const override;
-	virtual bool CanReceiveCaptionSamples(uint32 Num) const override;
-	virtual bool CanReceiveMetadataSamples(uint32 Num) const override;
+	MEDIAUTILS_API virtual bool CanReceiveVideoSamples(uint32 Num) const override;
+	MEDIAUTILS_API virtual bool CanReceiveAudioSamples(uint32 Num) const override;
+	MEDIAUTILS_API virtual bool CanReceiveSubtitleSamples(uint32 Num) const override;
+	MEDIAUTILS_API virtual bool CanReceiveCaptionSamples(uint32 Num) const override;
+	MEDIAUTILS_API virtual bool CanReceiveMetadataSamples(uint32 Num) const override;
 
 private:
 
diff --git a/Engine/Source/Runtime/MeshConversion/Public/DynamicMeshToMeshDescription.h b/Engine/Source/Runtime/MeshConversion/Public/DynamicMeshToMeshDescription.h
--- a/Engine/Source/Runtime/MeshConversion/Public/DynamicMeshToMeshDescription.h
+++ b/Engine/Source/Runtime/MeshConversion/Public/DynamicMeshToMeshDescription.h
@@ -21,7 +21,7 @@ using UE::Geometry::FDynamicMesh3;
  * Convert FDynamicMesh3 to FMeshDescription
  *
  */
-class MESHCONVERSION_API FDynamicMeshToMeshDescription
+class FDynamicMeshToMeshDescription
 {
 public:
 	/** If true, will print some possibly-helpful debugging spew to output log */
@@ -47,7 +47,7 @@ public:
 	 * @param bAttributesOnly If true, only check what needs to be checked for UpdateAttributes
 							 (will check vertices or triangles depending on whether attributes are per vertex or in overlays)
 	 */
-	static bool HaveMatchingElementCounts(const FDynamicMesh3* DynamicMesh, const FMeshDescription* MeshDescription, bool bVerticesOnly, bool bAttributesOnly);
+	static MESHCONVERSION_API bool HaveMatchingElementCounts(const FDynamicMesh3* DynamicMesh, const FMeshDescription* MeshDescription, bool bVerticesOnly, bool bAttributesOnly);
 
 	/**
 	 * Checks if element counts match. If false then Update can't be called -- you must call Convert
@@ -56,7 +56,7 @@ public:
 	 * @param DynamicMesh The dynamic mesh with updated vertices or attributes
 	 * @param MeshDescription The corresponding mesh description
 	 */
-	bool HaveMatchingElementCounts(const FDynamicMesh3* DynamicMesh, const FMeshDescription* MeshDescription);
+	MESHCONVERSION_API bool HaveMatchingElementCounts(const FDynamicMesh3* DynamicMesh, const FMeshDescription* MeshDescription);
 
 	/**
 	 * Default conversion of DynamicMesh to MeshDescription. Calls functions below depending on mesh state
@@ -67,7 +67,7 @@ public:
 	 *       in this case the MeshDescription is expected to have empty tangents.
 	 *       
 	 */
-	void Convert(const FDynamicMesh3* MeshIn, FMeshDescription& MeshOut, bool bCopyTangents = false);
+	MESHCONVERSION_API void Convert(const FDynamicMesh3* MeshIn, FMeshDescription& MeshOut, bool bCopyTangents = false);
 
 	/**
 	 * Updates the given mesh description based conversion options provided in the constructor. Assumes
@@ -75,14 +75,14 @@ public:
 	 * Annoyingly, this can't just be named Update() due to ambiguity with the function below, which
 	 * existed beforehand and should probably have been this function instead.
 	 */
-	void UpdateUsingConversionOptions(const FDynamicMesh3* MeshIn, FMeshDescription& MeshOut);
+	MESHCONVERSION_API void UpdateUsingConversionOptions(const FDynamicMesh3* MeshIn, FMeshDescription& MeshOut);
 
 	/**
 	 * Update existing MeshDescription based on DynamicMesh. Assumes mesh topology has not changed.
 	 * Copies positions 
 	 * optionally, normals, tangents and UVs
 	 */
-	void Update(const FDynamicMesh3* MeshIn, FMeshDescription& MeshOut, bool bUpdateNormals = true, bool bUpdateTangents = false, bool bUpdateUVs = false);
+	MESHCONVERSION_API void Update(const FDynamicMesh3* MeshIn, FMeshDescription& MeshOut, bool bUpdateNormals = true, bool bUpdateTangents = false, bool bUpdateUVs = false);
 	
 
 	/**
@@ -94,7 +94,7 @@ public:
 	 *                          this requires the PrimaryNormals to exist as well. 
 	 *  @param bUpdateUVs       Specifices if the UV layers should be transfered from the MeshIn overlays.
 	 */
-	void UpdateAttributes(const FDynamicMesh3* MeshIn, FMeshDescription& MeshOut, bool bUpdateNormals, bool bUpdateTangents, bool bUpdateUVs);
+	MESHCONVERSION_API void UpdateAttributes(const FDynamicMesh3* MeshIn, FMeshDescription& MeshOut, bool bUpdateNormals, bool bUpdateTangents, bool bUpdateUVs);
 
 	/**
 	 * Use the TMeshTangents to update the Tangent and BinormalSign attributes of the MeshDescription, assuming mesh topology has not changed. Does not modify any other attributes.
@@ -103,7 +103,7 @@ public:
 	 *	NOTE: assumes the order of triangles in the MeshIn correspond to the ordering you'd get by iterating over triangles, on MeshOut
 	 *		  This matches conversion currently used in MeshDescriptionToDynamicMesh.cpp, but if that changes we will need to change this function to match!
 	 */
-	void UpdateTangents(const FDynamicMesh3* MeshIn, FMeshDescription& MeshOut, const UE::Geometry::TMeshTangents<double>* SrcTangents);
+	MESHCONVERSION_API void UpdateTangents(const FDynamicMesh3* MeshIn, FMeshDescription& MeshOut, const UE::Geometry::TMeshTangents<double>* SrcTangents);
 
 
 	/**
@@ -111,7 +111,7 @@ public:
 	 *	NOTE: assumes the order of triangles in the MeshIn correspond to the ordering you'd get by iterating over triangles, on MeshOut
 	 *		  This matches conversion currently used in MeshDescriptionToDynamicMesh.cpp, but if that changes we will need to change this function to match!
 	 */
-	void UpdateTangents(const FDynamicMesh3* MeshIn, FMeshDescription& MeshOut);
+	MESHCONVERSION_API void UpdateTangents(const FDynamicMesh3* MeshIn, FMeshDescription& MeshOut);
 
 
 	/**
@@ -119,7 +119,7 @@ public:
 	 *	NOTE: assumes the order of triangles in the MeshIn correspond to the ordering you'd get by iterating over triangles, on MeshOut
 	 *		  This matches conversion currently used in MeshDescriptionToDynamicMesh.cpp, but if that changes we will need to change this function to match!
 	 */
-	void UpdateVertexColors(const FDynamicMesh3* MeshIn, FMeshDescription& MeshOut);
+	MESHCONVERSION_API void UpdateVertexColors(const FDynamicMesh3* MeshIn, FMeshDescription& MeshOut);
 
 
 
@@ -131,7 +131,7 @@ public:
 	 * Ignore any Attributes on input Mesh, calculate per-vertex normals and have MeshDescription compute tangents.
 	 * One VertexInstance per input vertex is generated
 	 */
-	void Convert_NoAttributes(const FDynamicMesh3* MeshIn, FMeshDescription& MeshOut);
+	MESHCONVERSION_API void Convert_NoAttributes(const FDynamicMesh3* MeshIn, FMeshDescription& MeshOut);
 
 	/**
 	 * Convert while minimizing VertexInstance count, IE new VertexInstances are only created 
@@ -140,7 +140,7 @@ public:
 	 * Note: This doesn't copy any tangents from the FDynamicMesh
 	 * Note: This conversion is not currently being used. It is unclear if all consumers of FMeshDescription can handle such shared vertex instances 
 	 */
-	void Convert_SharedInstances(const FDynamicMesh3* MeshIn, FMeshDescription& MeshOut);
+	MESHCONVERSION_API void Convert_SharedInstances(const FDynamicMesh3* MeshIn, FMeshDescription& MeshOut);
 
 	/**
 	 * Convert with no shared VertexInstances. A new VertexInstance is created for
@@ -152,7 +152,7 @@ public:
 	 * Note: Don't copy tangents if the resulting MeshDescription corresponds to a StaticMesh with autogenerated tangents
 	 *       in this case the MeshDescription is expected to have empty tangents.
 	 */
-	void Convert_NoSharedInstances(const FDynamicMesh3* MeshIn, FMeshDescription& MeshOut, bool bCopyTangents);
+	MESHCONVERSION_API void Convert_NoSharedInstances(const FDynamicMesh3* MeshIn, FMeshDescription& MeshOut, bool bCopyTangents);
 
 
 
@@ -163,14 +163,14 @@ protected:
 	 * Will copy to existing MeshDescription TriangleAttribute<int32> if one with the same name exists.
 	 * Otherwise will register a new one.
 	 */
-	void ConvertPolygroupLayers(const FDynamicMesh3* MeshIn, FMeshDescription& MeshOut, const TArray<FTriangleID>& IndexToTriangleIDMap);
+	MESHCONVERSION_API void ConvertPolygroupLayers(const FDynamicMesh3* MeshIn, FMeshDescription& MeshOut, const TArray<FTriangleID>& IndexToTriangleIDMap);
 
 	/**
 	* Transfer WeightLayers from DynamicMesh AttributeSet to MeshDescription.
 	* Will copy to existing MeshDescription VertexAttribute<float> if one with the same name exists.
 	* Otherwise will register a new one.
 	*/
-	void ConvertWeightLayers(const FDynamicMesh3* MeshIn, FMeshDescription& MeshOut, const TArray<FVertexID>& IndexToVertexIDMap);
+	MESHCONVERSION_API void ConvertWeightLayers(const FDynamicMesh3* MeshIn, FMeshDescription& MeshOut, const TArray<FVertexID>& IndexToVertexIDMap);
 
 	/**
 	 * Applies an optional sRGB-to-Linear color transform on the input. The color transform
@@ -181,5 +181,5 @@ protected:
 	 *
 	 * @param Color color to transform
 	 */
-	void ApplyVertexColorTransform(FVector4f& Color) const;
-};
\ No newline at end of file
+	MESHCONVERSION_API void ApplyVertexColorTransform(FVector4f& Color) const;
+};
diff --git a/Engine/Source/Runtime/MeshConversion/Public/MeshConversionOptions.h b/Engine/Source/Runtime/MeshConversion/Public/MeshConversionOptions.h
--- a/Engine/Source/Runtime/MeshConversion/Public/MeshConversionOptions.h
+++ b/Engine/Source/Runtime/MeshConversion/Public/MeshConversionOptions.h
@@ -5,7 +5,7 @@
 #include "CoreMinimal.h"
 
 
-struct MESHCONVERSION_API FConversionToMeshDescriptionOptions
+struct FConversionToMeshDescriptionOptions
 {
 public:
 	/** Should triangle groups be transfered to MeshDescription via custom PolyTriGroups attribute */
@@ -39,4 +39,4 @@ public:
 		bSetPolyGroups = bUpdatePositions = bUpdateNormals = bUpdateTangents = bUpdateUVs = false;
 		bUpdateVtxColors = true;
 	}
-};
\ No newline at end of file
+};
diff --git a/Engine/Source/Runtime/MeshConversion/Public/MeshDescriptionBuilder.h b/Engine/Source/Runtime/MeshConversion/Public/MeshDescriptionBuilder.h
--- a/Engine/Source/Runtime/MeshConversion/Public/MeshDescriptionBuilder.h
+++ b/Engine/Source/Runtime/MeshConversion/Public/MeshDescriptionBuilder.h
@@ -31,81 +31,81 @@ namespace ExtendedMeshAttribute
  * Utility class to construct MeshDescription instances.
  * NB: this will add a vertex-instance UV to the Description, if it has none. 
  */
-class MESHCONVERSION_API FMeshDescriptionBuilder
+class FMeshDescriptionBuilder
 {
 public:
-	void SetMeshDescription(FMeshDescription* Description);
+	MESHCONVERSION_API void SetMeshDescription(FMeshDescription* Description);
 
 	/** Pre-allocate space in the mesh description */
-	void ReserveNewVertices(int32 Count);
+	MESHCONVERSION_API void ReserveNewVertices(int32 Count);
 
 	/** Append vertex and return new vertex ID */
-	FVertexID AppendVertex(const FVector& Position);
+	MESHCONVERSION_API FVertexID AppendVertex(const FVector& Position);
 
 	/** Return position of vertex */
-	FVector GetPosition(const FVertexID& VertexID);
+	MESHCONVERSION_API FVector GetPosition(const FVertexID& VertexID);
 
 	/** Return position of vertex parent of instance */
-	FVector GetPosition(const FVertexInstanceID& InstanceID);
+	MESHCONVERSION_API FVector GetPosition(const FVertexInstanceID& InstanceID);
 
 	/** Set the position of a vertex */
-	void SetPosition(const FVertexID& VertexID, const FVector& NewPosition);
+	MESHCONVERSION_API void SetPosition(const FVertexID& VertexID, const FVector& NewPosition);
 
 	/** Set the number of UV layers */
-	void SetNumUVLayers(int32 NumUVLayers);
+	MESHCONVERSION_API void SetNumUVLayers(int32 NumUVLayers);
 
 	/** Pre-allocate space in the mesh description for UVs in the indicated UVLayer */
-	void ReserveNewUVs(int32 Count, int UVLayerIndex);
+	MESHCONVERSION_API void ReserveNewUVs(int32 Count, int UVLayerIndex);
 
 	/** Append a UV 'vertex' and return a new UV ID*/
-	FUVID AppendUV(const FVector2D& UVvalue, int32 UVLayerIndex);
+	MESHCONVERSION_API FUVID AppendUV(const FVector2D& UVvalue, int32 UVLayerIndex);
 
 	/** Append new vertex instance and return ID */
-	FVertexInstanceID AppendInstance(const FVertexID& VertexID);
+	MESHCONVERSION_API FVertexInstanceID AppendInstance(const FVertexID& VertexID);
 
 	/** 
 	* Set the UV of a vertex instance.
 	* Note: this generally shouldn't be called directly because it alters the instance UVs but not the shared UVs
 	* When setting UVs use AppendUVTriangle	  
 	*/
-	void SetInstanceUV(const FVertexInstanceID& InstanceID, const FVector2D& InstanceUV, int32 UVLayerIndex = 0);
+	MESHCONVERSION_API void SetInstanceUV(const FVertexInstanceID& InstanceID, const FVector2D& InstanceUV, int32 UVLayerIndex = 0);
 
 	/** Set the Normal of a vertex instance*/
-	void SetInstanceNormal(const FVertexInstanceID& InstanceID, const FVector& Normal);
+	MESHCONVERSION_API void SetInstanceNormal(const FVertexInstanceID& InstanceID, const FVector& Normal);
 
 	/** Set the full tangent space of a vertex instance, this is stored as a vec3 normal, vec3 tangent, and a bitangent sign */
-	void SetInstanceTangentSpace(const FVertexInstanceID& InstanceID, const FVector& Normal, const FVector& Tangent, float Sign);
+	MESHCONVERSION_API void SetInstanceTangentSpace(const FVertexInstanceID& InstanceID, const FVector& Normal, const FVector& Tangent, float Sign);
 
 	/** Set the Color of a vertex instance*/
-	void SetInstanceColor(const FVertexInstanceID& InstanceID, const FVector4f& Color);
+	MESHCONVERSION_API void SetInstanceColor(const FVertexInstanceID& InstanceID, const FVector4f& Color);
 
 	/** Enable per-triangle integer attribute named PolyTriGroups */
-	void EnablePolyGroups();
+	MESHCONVERSION_API void EnablePolyGroups();
 
 	/** Set the PolyTriGroups attribute value to a specific GroupID for a Triangle */
-	void SetPolyGroupID(const FTriangleID& TriangleID, int GroupID);
+	MESHCONVERSION_API void SetPolyGroupID(const FTriangleID& TriangleID, int GroupID);
 
 	/** Append a UV triangle to the specified UV layer. This will use both shared and per-instance UV storage*/
-	void AppendUVTriangle(const FTriangleID& TriangleID, const FUVID UVverterxID0, const FUVID UVvertexID1, const FUVID UVvertexID2, int32 UVLayerIndex);
+	MESHCONVERSION_API void AppendUVTriangle(const FTriangleID& TriangleID, const FUVID UVverterxID0, const FUVID UVvertexID1, const FUVID UVvertexID2, int32 UVLayerIndex);
 
 	/** Append a triangle to the mesh with the given PolygonGroup ID */
-	FTriangleID AppendTriangle(const FVertexID& Vertex0, const FVertexID& Vertex1, const FVertexID& Vertex2, const FPolygonGroupID& PolygonGroup);
+	MESHCONVERSION_API FTriangleID AppendTriangle(const FVertexID& Vertex0, const FVertexID& Vertex1, const FVertexID& Vertex2, const FPolygonGroupID& PolygonGroup);
 
 
 	/** Append a triangle to the mesh with the given PolygonGroup ID */
-	FTriangleID AppendTriangle(const FVertexID* Triangle, const FPolygonGroupID& PolygonGroup);
+	MESHCONVERSION_API FTriangleID AppendTriangle(const FVertexID* Triangle, const FPolygonGroupID& PolygonGroup);
 
 
 	/**
 	 * Append a triangle to the mesh using the given vertex instances and PolygonGroup ID
 	 */
-	FTriangleID AppendTriangle(const FVertexInstanceID& Instance0, const FVertexInstanceID& Instance1, const FVertexInstanceID& Instance2, const FPolygonGroupID& PolygonGroup);
+	MESHCONVERSION_API FTriangleID AppendTriangle(const FVertexInstanceID& Instance0, const FVertexInstanceID& Instance1, const FVertexInstanceID& Instance2, const FPolygonGroupID& PolygonGroup);
 
 	/**
 	 * Append an arbitrary polygon to the mesh with the given PolygonGroup ID
 	 * Unique Vertex instances will be created for each polygon-vertex.
 	 */
-	FPolygonID AppendPolygon(const TArray<FVertexID>& Vertices, const FPolygonGroupID& PolygonGroup);
+	MESHCONVERSION_API FPolygonID AppendPolygon(const TArray<FVertexID>& Vertices, const FPolygonGroupID& PolygonGroup);
 
 
 	/** 
@@ -113,7 +113,7 @@ public:
 	 * PolygonGroups are not the same as Polygroups, they essentially represent
 	 * Mesh Sections, which then reference Material Slots, etc
 	 */
-	FPolygonGroupID AppendPolygonGroup(FName MaterialSlotName = NAME_None);
+	MESHCONVERSION_API FPolygonGroupID AppendPolygonGroup(FName MaterialSlotName = NAME_None);
 
 	/**
 	 * Set the specified value for the named attribute / vertex index combination.
@@ -133,20 +133,20 @@ public:
 	}
 
 	/** Set MeshAttribute::Edge::IsHard to true for all edges */
-	void SetAllEdgesHardness(bool bHard);
+	MESHCONVERSION_API void SetAllEdgesHardness(bool bHard);
 
 	/** Translate the MeshDescription vertex positions */
-	void Translate(const FVector& Translation);
+	MESHCONVERSION_API void Translate(const FVector& Translation);
 
 
 	/** Return the current bounding box of the mesh */
-	FBox ComputeBoundingBox() const;
+	MESHCONVERSION_API FBox ComputeBoundingBox() const;
 
 	/** Disable the construction of secondary data structures in the mesh description */
-	void SuspendMeshDescriptionIndexing();
+	MESHCONVERSION_API void SuspendMeshDescriptionIndexing();
 
 	/** Enable the construction of secondary data structures in the mesh description */
-	void ResumeMeshDescriptionIndexing();
+	MESHCONVERSION_API void ResumeMeshDescriptionIndexing();
 
 protected:
 
@@ -166,4 +166,4 @@ protected:
 	TPolygonAttributesRef<int> PolyGroups;
 
 	TPolygonGroupAttributesRef<FName> GroupMaterialSlotNames;
-};
\ No newline at end of file
+};
diff --git a/Engine/Source/Runtime/MeshConversion/Public/MeshDescriptionToDynamicMesh.h b/Engine/Source/Runtime/MeshConversion/Public/MeshDescriptionToDynamicMesh.h
--- a/Engine/Source/Runtime/MeshConversion/Public/MeshDescriptionToDynamicMesh.h
+++ b/Engine/Source/Runtime/MeshConversion/Public/MeshDescriptionToDynamicMesh.h
@@ -24,7 +24,7 @@ using UE::Geometry::FDynamicMesh3;
  * @todo handle additional UV layers on MD
  * @todo option to disable UV/Normal welding
  */
-class MESHCONVERSION_API FMeshDescriptionToDynamicMesh
+class FMeshDescriptionToDynamicMesh
 {
 public:
 	/** If true, will print some possibly-helpful debugging spew to output log */
@@ -84,19 +84,19 @@ public:
 	 * @param bCopyTangents  - if bDisableAttributes is false, this requests the tangent plane vectors (tangent and bitangent) 
 	 *                          be stored as overlays in the MeshOut DynamicAttributeSet, provided they exist on the MeshIn
 	 */
-	void Convert(const FMeshDescription* MeshIn, FDynamicMesh3& MeshOut, bool bCopyTangents = false);
+	MESHCONVERSION_API void Convert(const FMeshDescription* MeshIn, FDynamicMesh3& MeshOut, bool bCopyTangents = false);
 
 	/**
 	 * Copy tangents from MeshDescription to a FMeshTangents instance.
 	 * @warning Convert() must have been used to create the TargetMesh before calling this function
 	 */
-	void CopyTangents(const FMeshDescription* SourceMesh, const FDynamicMesh3* TargetMesh, UE::Geometry::TMeshTangents<float>* TangentsOut);
+	MESHCONVERSION_API void CopyTangents(const FMeshDescription* SourceMesh, const FDynamicMesh3* TargetMesh, UE::Geometry::TMeshTangents<float>* TangentsOut);
 
 	/**
 	 * Copy tangents from MeshDescription to a FMeshTangents instance.
 	 * @warning Convert() must have been used to create the TargetMesh before calling this function
 	 */
-	void CopyTangents(const FMeshDescription* SourceMesh, const FDynamicMesh3* TargetMesh, UE::Geometry::TMeshTangents<double>* TangentsOut);
+	MESHCONVERSION_API void CopyTangents(const FMeshDescription* SourceMesh, const FDynamicMesh3* TargetMesh, UE::Geometry::TMeshTangents<double>* TangentsOut);
 
 protected:
 	/**
@@ -108,5 +108,5 @@ protected:
 	 *
 	 * @param Color color to transform
 	 */
-	void ApplyVertexColorTransform(FVector4f& Color) const;
-};
\ No newline at end of file
+	MESHCONVERSION_API void ApplyVertexColorTransform(FVector4f& Color) const;
+};
diff --git a/Engine/Source/Runtime/MeshConversion/Public/MeshDescriptionUVsToDynamicMesh.h b/Engine/Source/Runtime/MeshConversion/Public/MeshDescriptionUVsToDynamicMesh.h
--- a/Engine/Source/Runtime/MeshConversion/Public/MeshDescriptionUVsToDynamicMesh.h
+++ b/Engine/Source/Runtime/MeshConversion/Public/MeshDescriptionUVsToDynamicMesh.h
@@ -28,7 +28,7 @@ class FDynamicMesh3;
  * The converter prefers shared UVs when they are available, though it will populate instanced
  * UVs as well on the way back.
  */
-class MESHCONVERSION_API FMeshDescriptionUVsToDynamicMesh
+class FMeshDescriptionUVsToDynamicMesh
 {
 public:
 	int32 UVLayerIndex = 0;
@@ -36,13 +36,13 @@ public:
 	// The [0,1] UV range will be scaled to [0, ScaleFactor]
 	double ScaleFactor = 1000;
 
-	int32 GetNumUVLayers(const FMeshDescription* InputMeshDescription) const;
+	MESHCONVERSION_API int32 GetNumUVLayers(const FMeshDescription* InputMeshDescription) const;
 
 	// TODO: The conversion function forward to a dynamic mesh should be made to produce maps 
 	// similar to MeshDescriptionToDynamicMesh, which is why this is not a const function.
-	TSharedPtr<FDynamicMesh3> GetUVMesh(const FMeshDescription* InputMeshDescription);
+	MESHCONVERSION_API TSharedPtr<FDynamicMesh3> GetUVMesh(const FMeshDescription* InputMeshDescription);
 
-	void BakeBackUVsFromUVMesh(const FDynamicMesh3* DynamicUVMesh, FMeshDescription* OutputMeshDescription) const;
+	MESHCONVERSION_API void BakeBackUVsFromUVMesh(const FDynamicMesh3* DynamicUVMesh, FMeshDescription* OutputMeshDescription) const;
 };
 
-}}//end namespace UE::Geometry
\ No newline at end of file
+}}//end namespace UE::Geometry
diff --git a/Engine/Source/Runtime/MeshConversionEngineTypes/Public/SkeletalMeshLODModelToDynamicMesh.h b/Engine/Source/Runtime/MeshConversionEngineTypes/Public/SkeletalMeshLODModelToDynamicMesh.h
--- a/Engine/Source/Runtime/MeshConversionEngineTypes/Public/SkeletalMeshLODModelToDynamicMesh.h
+++ b/Engine/Source/Runtime/MeshConversionEngineTypes/Public/SkeletalMeshLODModelToDynamicMesh.h
@@ -14,7 +14,7 @@ using UE::Geometry::FDynamicMesh3;
  * Convert FSkeletalMeshLODModel to FDynamicMesh3
  *
  */
-class MESHCONVERSIONENGINETYPES_API FSkeletalMeshLODModelToDynamicMesh
+class FSkeletalMeshLODModelToDynamicMesh
 {
 public:
 	/** If true, will print some possibly-helpful debugging spew to output log */
@@ -47,6 +47,6 @@ public:
 	 * @param bCopyTangents  - if bDisableAttributes is false, this requests the tangent plane vectors (tangent and bitangent)
 	 *                          be stored as overlays in the MeshOut DynamicAttributeSet, provided they exist on the MeshIn
 	 */
-	void Convert(const  FSkeletalMeshLODModel* MeshIn, FDynamicMesh3& MeshOut, bool bCopyTangents = false);
+	MESHCONVERSIONENGINETYPES_API void Convert(const  FSkeletalMeshLODModel* MeshIn, FDynamicMesh3& MeshOut, bool bCopyTangents = false);
 
-};
\ No newline at end of file
+};
diff --git a/Engine/Source/Runtime/MeshConversionEngineTypes/Public/StaticMeshLODResourcesToDynamicMesh.h b/Engine/Source/Runtime/MeshConversionEngineTypes/Public/StaticMeshLODResourcesToDynamicMesh.h
--- a/Engine/Source/Runtime/MeshConversionEngineTypes/Public/StaticMeshLODResourcesToDynamicMesh.h
+++ b/Engine/Source/Runtime/MeshConversionEngineTypes/Public/StaticMeshLODResourcesToDynamicMesh.h
@@ -20,11 +20,11 @@ namespace Geometry
  * split UV/normal/tangent/color, ie in the overlays there will be a unique overlay element for each
  * base mesh vertex.
  */
-class MESHCONVERSIONENGINETYPES_API FStaticMeshLODResourcesToDynamicMesh
+class FStaticMeshLODResourcesToDynamicMesh
 {
 public:
 
-	struct MESHCONVERSIONENGINETYPES_API ConversionOptions
+	struct ConversionOptions
 	{
 		bool bWantNormals = true;
 		bool bWantTangents = true;
@@ -36,7 +36,7 @@ public:
 		FVector3d BuildScale = FVector3d::One();
 	};
 
-	static bool Convert(
+	static MESHCONVERSIONENGINETYPES_API bool Convert(
 		const FStaticMeshLODResources* StaticMeshResources, 
 		const ConversionOptions& Options, 
 		FDynamicMesh3& OutputMesh);
diff --git a/Engine/Source/Runtime/MeshDescription/Public/MeshAttributes.h b/Engine/Source/Runtime/MeshDescription/Public/MeshAttributes.h
--- a/Engine/Source/Runtime/MeshDescription/Public/MeshAttributes.h
+++ b/Engine/Source/Runtime/MeshDescription/Public/MeshAttributes.h
@@ -55,7 +55,7 @@ namespace MeshAttribute
 }
 
 
-class MESHDESCRIPTION_API FMeshAttributes
+class FMeshAttributes
 {
 public:
 	explicit FMeshAttributes(FMeshDescription& InMeshDescription)
@@ -64,7 +64,7 @@ public:
 
 	virtual ~FMeshAttributes() = default;
 
-	virtual void Register(bool bKeepExistingAttribute = false);
+	MESHDESCRIPTION_API virtual void Register(bool bKeepExistingAttribute = false);
 
 	static bool IsReservedAttributeName(const FName InAttributeName)
 	{
diff --git a/Engine/Source/Runtime/MeshDescription/Public/MeshDescription.h b/Engine/Source/Runtime/MeshDescription/Public/MeshDescription.h
--- a/Engine/Source/Runtime/MeshDescription/Public/MeshDescription.h
+++ b/Engine/Source/Runtime/MeshDescription/Public/MeshDescription.h
@@ -90,16 +90,16 @@ enum class EComputeNTBsOptions : uint32
 ENUM_CLASS_FLAGS(EComputeNTBsOptions);
 
 
-struct MESHDESCRIPTION_API FMeshDescription
+struct FMeshDescription
 {
 public:
 
 	// Mesh description should be a moveable type.
 	// Hence explicitly declare all the below as defaulted, to ensure they will be generated by the compiler.
-	FMeshDescription();
+	MESHDESCRIPTION_API FMeshDescription();
 	~FMeshDescription() = default;
-	FMeshDescription(const FMeshDescription&);
-	FMeshDescription& operator=(const FMeshDescription&);
+	MESHDESCRIPTION_API FMeshDescription(const FMeshDescription&);
+	MESHDESCRIPTION_API FMeshDescription& operator=(const FMeshDescription&);
 	FMeshDescription(FMeshDescription&&) = default;
 	FMeshDescription& operator=(FMeshDescription&&) = default;
 
@@ -110,21 +110,21 @@ public:
 	}
 
 	// Serialize the mesh description
-	void Serialize(FArchive& Ar);
+	MESHDESCRIPTION_API void Serialize(FArchive& Ar);
 
 	// Legacy serialization for old assets
-	void SerializeLegacy(FArchive& Ar);
+	MESHDESCRIPTION_API void SerializeLegacy(FArchive& Ar);
 
 	// Empty the meshdescription
-	void Empty();
+	MESHDESCRIPTION_API void Empty();
 
 	// Operations on all the indexers
-	void ResetIndexers();
-	void BuildIndexers();
-	void RebuildIndexers();
+	MESHDESCRIPTION_API void ResetIndexers();
+	MESHDESCRIPTION_API void BuildIndexers();
+	MESHDESCRIPTION_API void RebuildIndexers();
 
 	// Return whether the mesh description is empty
-	bool IsEmpty() const;
+	MESHDESCRIPTION_API bool IsEmpty() const;
 
 	FVertexArray& Vertices() { return static_cast<TMeshElementContainer<FVertexID>&>(VertexElements->Get()); }
 	const FVertexArray& Vertices() const { return static_cast<const TMeshElementContainer<FVertexID>&>(VertexElements->Get()); }
@@ -251,7 +251,7 @@ public:
 	}
 
 	/** Deletes a vertex instance from a mesh */
-	void DeleteVertexInstance(const FVertexInstanceID VertexInstanceID, TArray<FVertexID>* InOutOrphanedVerticesPtr = nullptr);
+	MESHDESCRIPTION_API void DeleteVertexInstance(const FVertexInstanceID VertexInstanceID, TArray<FVertexID>* InOutOrphanedVerticesPtr = nullptr);
 	
 	/** Returns whether the passed vertex instance ID is valid */
 	bool IsVertexInstanceValid(const FVertexInstanceID VertexInstanceID) const
@@ -314,7 +314,7 @@ public:
 	}
 
 	/** Deletes an edge from the mesh */
-	void DeleteEdge(const FEdgeID EdgeID, TArray<FVertexID>* InOutOrphanedVerticesPtr = nullptr);
+	MESHDESCRIPTION_API void DeleteEdge(const FEdgeID EdgeID, TArray<FVertexID>* InOutOrphanedVerticesPtr = nullptr);
 
 	/** Returns whether the passed edge ID is valid */
 	bool IsEdgeValid(const FEdgeID EdgeID) const
@@ -344,12 +344,12 @@ public:
 	}
 
 	/** Deletes a triangle from the mesh */
-	void DeleteTriangle(const FTriangleID TriangleID, TArray<FEdgeID>* InOutOrphanedEdgesPtr = nullptr, TArray<FVertexInstanceID>* InOutOrphanedVertexInstancesPtr = nullptr, TArray<FPolygonGroupID>* InOutOrphanedPolygonGroupsPtr = nullptr);
+	MESHDESCRIPTION_API void DeleteTriangle(const FTriangleID TriangleID, TArray<FEdgeID>* InOutOrphanedEdgesPtr = nullptr, TArray<FVertexInstanceID>* InOutOrphanedVertexInstancesPtr = nullptr, TArray<FPolygonGroupID>* InOutOrphanedPolygonGroupsPtr = nullptr);
 
 	/** Deletes triangles from the mesh and remove all orphaned polygon groups, vertex instances, edges and vertices.
 		Will not compact the internal arrays, you must call Compact() manually.
 	 */
-	void DeleteTriangles(const TArray<FTriangleID>& Triangles);
+	MESHDESCRIPTION_API void DeleteTriangles(const TArray<FTriangleID>& Triangles);
 
 	/** Returns whether the passed triangle ID is valid */
 	bool IsTriangleValid(const FTriangleID TriangleID) const
@@ -379,12 +379,12 @@ public:
 	}
 
 	/** Deletes a polygon from the mesh */
-	void DeletePolygon(const FPolygonID PolygonID, TArray<FEdgeID>* InOutOrphanedEdgesPtr = nullptr, TArray<FVertexInstanceID>* InOutOrphanedVertexInstancesPtr = nullptr, TArray<FPolygonGroupID>* InOutOrphanedPolygonGroupsPtr = nullptr);
+	MESHDESCRIPTION_API void DeletePolygon(const FPolygonID PolygonID, TArray<FEdgeID>* InOutOrphanedEdgesPtr = nullptr, TArray<FVertexInstanceID>* InOutOrphanedVertexInstancesPtr = nullptr, TArray<FPolygonGroupID>* InOutOrphanedPolygonGroupsPtr = nullptr);
 
 	/** Deletes polygons from the mesh and remove all orphaned polygon groups, vertex instances, edges and vertices.
 		Will not compact the internal arrays, you must call Compact() manually.
 	 */
-	void DeletePolygons(const TArray<FPolygonID>& Polygons);
+	MESHDESCRIPTION_API void DeletePolygons(const TArray<FPolygonID>& Polygons);
 
 	/** Returns whether the passed polygon ID is valid */
 	bool IsPolygonValid(const FPolygonID PolygonID) const
@@ -436,10 +436,10 @@ public:
 	// Vertex operations
 
 	/** Returns whether a given vertex is orphaned, i.e. it doesn't form part of any polygon */
-	bool IsVertexOrphaned(const FVertexID VertexID) const;
+	MESHDESCRIPTION_API bool IsVertexOrphaned(const FVertexID VertexID) const;
 
 	/** Returns the edge ID defined by the two given vertex IDs, if there is one; otherwise INDEX_NONE */
-	FEdgeID GetVertexPairEdge(const FVertexID VertexID0, const FVertexID VertexID1) const;
+	MESHDESCRIPTION_API FEdgeID GetVertexPairEdge(const FVertexID VertexID0, const FVertexID VertexID1) const;
 
 	/** Returns reference to an array of Edge IDs connected to this vertex */
 	TArrayView<const FEdgeID> GetVertexConnectedEdgeIDs(const FVertexID VertexID) const
@@ -611,7 +611,7 @@ public:
 	}
 
 	/** Returns the edge ID defined by the two given vertex instance IDs, if there is one; otherwise INDEX_NONE */
-	FEdgeID GetVertexInstancePairEdge(const FVertexInstanceID VertexInstanceID0, const FVertexInstanceID VertexInstanceID1) const;
+	MESHDESCRIPTION_API FEdgeID GetVertexInstancePairEdge(const FVertexInstanceID VertexInstanceID0, const FVertexInstanceID VertexInstanceID1) const;
 
 	/** Returns reference to an array of Triangle IDs connected to this vertex instance */
 	TArrayView<const FTriangleID> GetVertexInstanceConnectedTriangleIDs(const FVertexInstanceID VertexInstanceID) const
@@ -875,10 +875,10 @@ public:
 	}
 
 	/** Reverse the winding order of the vertices of this triangle */
-	void ReverseTriangleFacing(const FTriangleID TriangleID);
+	MESHDESCRIPTION_API void ReverseTriangleFacing(const FTriangleID TriangleID);
 
 	/** Set the UV indices for this triangle */
-	void SetTriangleUVIndices(const FTriangleID TriangleID, TArrayView<const FUVID> UVIDs, int32 UVChannel = 0);
+	MESHDESCRIPTION_API void SetTriangleUVIndices(const FTriangleID TriangleID, TArrayView<const FUVID> UVIDs, int32 UVChannel = 0);
 
 	//////////////////////////////////////////////////////////////////////
 	// Polygon operations
@@ -1125,8 +1125,8 @@ public:
 	}
 
 	/** Set the vertex instance at the given index around the polygon to the new value */
-	void SetPolygonVertexInstance(const FPolygonID PolygonID, const int32 PerimeterIndex, const FVertexInstanceID VertexInstanceID);
-	void SetPolygonVertexInstances(const FPolygonID PolygonID, TArrayView<const FVertexInstanceID> VertexInstanceIDs);
+	MESHDESCRIPTION_API void SetPolygonVertexInstance(const FPolygonID PolygonID, const int32 PerimeterIndex, const FVertexInstanceID VertexInstanceID);
+	MESHDESCRIPTION_API void SetPolygonVertexInstances(const FPolygonID PolygonID, TArrayView<const FVertexInstanceID> VertexInstanceIDs);
 
 	/** Sets the polygon group associated with a polygon */
 	void SetPolygonPolygonGroup(const FPolygonID PolygonID, const FPolygonGroupID PolygonGroupID)
@@ -1149,14 +1149,14 @@ public:
 	}
 
 	/** Reverse the winding order of the vertices of this polygon */
-	void ReversePolygonFacing(const FPolygonID PolygonID);
+	MESHDESCRIPTION_API void ReversePolygonFacing(const FPolygonID PolygonID);
 
 	/** Determines the vertex instances which form the perimeter of a polygon */
-	void FindPolygonPerimeter(TArrayView<const FTriangleID> Triangles, TArrayView<TTuple<int32, int32>> Result) const;
-	void FindPolygonPerimeter(const FPolygonID PolygonID, TArrayView<FEdgeID> Edges) const;
+	MESHDESCRIPTION_API void FindPolygonPerimeter(TArrayView<const FTriangleID> Triangles, TArrayView<TTuple<int32, int32>> Result) const;
+	MESHDESCRIPTION_API void FindPolygonPerimeter(const FPolygonID PolygonID, TArrayView<FEdgeID> Edges) const;
 
 	/** Generates triangles and internal edges for the given polygon */
-	void ComputePolygonTriangulation(const FPolygonID PolygonID);
+	MESHDESCRIPTION_API void ComputePolygonTriangulation(const FPolygonID PolygonID);
 
 
 	//////////////////////////////////////////////////////////////////////
@@ -1193,60 +1193,60 @@ public:
 	}
 
 	/** Remaps polygon groups according to the supplied map */
-	void RemapPolygonGroups(const TMap<FPolygonGroupID, FPolygonGroupID>& Remap);
+	MESHDESCRIPTION_API void RemapPolygonGroups(const TMap<FPolygonGroupID, FPolygonGroupID>& Remap);
 	
 	/** Transfer the source polygon group data to the destination polygon group. The source data is append to the destination */
-	void TransferPolygonGroup(FPolygonGroupID SourceID, FPolygonGroupID DestinationID);
+	MESHDESCRIPTION_API void TransferPolygonGroup(FPolygonGroupID SourceID, FPolygonGroupID DestinationID);
 
 	//////////////////////////////////////////////////////////////////////
 	// Whole mesh operations
 
 	/** Determines if calling Compact() would perform actual compaction or not. */
-	bool NeedsCompact() const;
+	MESHDESCRIPTION_API bool NeedsCompact() const;
 
 	/** Compacts the data held in the mesh description, and returns an object describing how the IDs have been remapped. */
-	void Compact( FElementIDRemappings& OutRemappings );
+	MESHDESCRIPTION_API void Compact( FElementIDRemappings& OutRemappings );
 
 	/** Remaps the element IDs in the mesh description according to the passed in object */
-	void Remap( const FElementIDRemappings& Remappings );
+	MESHDESCRIPTION_API void Remap( const FElementIDRemappings& Remappings );
 
 	/** Gets the number of UV element channels */
 	int32 GetNumUVElementChannels() const { return UVElements->GetNumChannels(); }
 
 	/** Sets the specified number of UV channels */
-	void SetNumUVChannels(const int32 NumUVChannels);
+	MESHDESCRIPTION_API void SetNumUVChannels(const int32 NumUVChannels);
 
 	/** Returns bounds of vertices */
-	FBoxSphereBounds GetBounds() const;
+	MESHDESCRIPTION_API FBoxSphereBounds GetBounds() const;
 
 	/** Retriangulates the entire mesh */
-	void TriangulateMesh();
+	MESHDESCRIPTION_API void TriangulateMesh();
 
 	/** Reverses the winding order of all polygons in the mesh */
-	void ReverseAllPolygonFacing();
+	MESHDESCRIPTION_API void ReverseAllPolygonFacing();
 
-	float GetTriangleCornerAngleForVertex(const FTriangleID TriangleID, const FVertexID VertexID) const;
-	float GetPolygonCornerAngleForVertex(const FPolygonID PolygonID, const FVertexID VertexID) const;
+	MESHDESCRIPTION_API float GetTriangleCornerAngleForVertex(const FTriangleID TriangleID, const FVertexID VertexID) const;
+	MESHDESCRIPTION_API float GetPolygonCornerAngleForVertex(const FPolygonID PolygonID, const FVertexID VertexID) const;
 
-	FBox ComputeBoundingBox() const;
+	MESHDESCRIPTION_API FBox ComputeBoundingBox() const;
 
 private:
 
-	FPlane ComputePolygonPlane(TArrayView<const FVertexInstanceID> PerimeterVertexInstanceIDs) const;
-	FVector ComputePolygonNormal(TArrayView<const FVertexInstanceID> PerimeterVertexInstanceIDs) const;
+	MESHDESCRIPTION_API FPlane ComputePolygonPlane(TArrayView<const FVertexInstanceID> PerimeterVertexInstanceIDs) const;
+	MESHDESCRIPTION_API FVector ComputePolygonNormal(TArrayView<const FVertexInstanceID> PerimeterVertexInstanceIDs) const;
 
 private:
 
-	void Initialize();
-	void InitializeIndexers();
-	void Cache();
+	MESHDESCRIPTION_API void Initialize();
+	MESHDESCRIPTION_API void InitializeIndexers();
+	MESHDESCRIPTION_API void Cache();
 
-	void CreateVertexInstance_Internal(const FVertexInstanceID VertexInstanceID, const FVertexID VertexID);
-	void CreateEdge_Internal(const FEdgeID EdgeID, const FVertexID VertexID0, const FVertexID VertexID1);
-	void CreateTriangle_Internal(const FTriangleID TriangleID, const FPolygonGroupID PolygonGroupID, TArrayView<const FVertexInstanceID> VertexInstanceIDs, TArray<FEdgeID>* OutEdgeIDs);
-	void CreatePolygon_Internal(const FPolygonID PolygonID, const FPolygonGroupID PolygonGroupID, TArrayView<const FVertexInstanceID> VertexInstanceIDs, TArray<FEdgeID>* OutEdgeIDs);
-	void CreatePolygonTriangles(const FPolygonID PolygonID, TArrayView<const FVertexInstanceID> VertexInstanceIDs);
-	void RemovePolygonTriangles(const FPolygonID PolygonID);
+	MESHDESCRIPTION_API void CreateVertexInstance_Internal(const FVertexInstanceID VertexInstanceID, const FVertexID VertexID);
+	MESHDESCRIPTION_API void CreateEdge_Internal(const FEdgeID EdgeID, const FVertexID VertexID0, const FVertexID VertexID1);
+	MESHDESCRIPTION_API void CreateTriangle_Internal(const FTriangleID TriangleID, const FPolygonGroupID PolygonGroupID, TArrayView<const FVertexInstanceID> VertexInstanceIDs, TArray<FEdgeID>* OutEdgeIDs);
+	MESHDESCRIPTION_API void CreatePolygon_Internal(const FPolygonID PolygonID, const FPolygonGroupID PolygonGroupID, TArrayView<const FVertexInstanceID> VertexInstanceIDs, TArray<FEdgeID>* OutEdgeIDs);
+	MESHDESCRIPTION_API void CreatePolygonTriangles(const FPolygonID PolygonID, TArrayView<const FVertexInstanceID> VertexInstanceIDs);
+	MESHDESCRIPTION_API void RemovePolygonTriangles(const FPolygonID PolygonID);
 
 	template <template <typename...> class TContainer>
 	void DeleteVertexInstance_Internal(const FVertexInstanceID VertexInstanceID, TContainer<FVertexID>* InOutOrphanedVerticesPtr = nullptr);
@@ -1258,7 +1258,7 @@ private:
 	void DeletePolygon_Internal(const FPolygonID PolygonID, TContainer<FEdgeID>* InOutOrphanedEdgesPtr = nullptr, TContainer<FVertexInstanceID>* InOutOrphanedVertexInstancesPtr = nullptr, TContainer<FPolygonGroupID>* InOutOrphanedPolygonGroupsPtr = nullptr);
 
 	/** Given a set of index remappings, fixes up references to element IDs */
-	void FixUpElementIDs(const FElementIDRemappings& Remappings);
+	MESHDESCRIPTION_API void FixUpElementIDs(const FElementIDRemappings& Remappings);
 
 	template <class T>
 	void AddUnique(TArray<T>& Container, const T& Item)
@@ -1273,13 +1273,13 @@ private:
 	}
 
 public:
-	static FName VerticesName;
-	static FName VertexInstancesName;
-	static FName UVsName;
-	static FName EdgesName;
-	static FName TrianglesName;
-	static FName PolygonsName;
-	static FName PolygonGroupsName;
+	static MESHDESCRIPTION_API FName VerticesName;
+	static MESHDESCRIPTION_API FName VertexInstancesName;
+	static MESHDESCRIPTION_API FName UVsName;
+	static MESHDESCRIPTION_API FName EdgesName;
+	static MESHDESCRIPTION_API FName TrianglesName;
+	static MESHDESCRIPTION_API FName PolygonsName;
+	static MESHDESCRIPTION_API FName PolygonGroupsName;
 
 private:
 	TMap<FName, FMeshElementTypeWrapper> Elements;
@@ -1319,7 +1319,7 @@ private:
 /**
  * Bulk data storage for FMeshDescription
  */
-struct MESHDESCRIPTION_API FMeshDescriptionBulkData
+struct FMeshDescriptionBulkData
 {
 public:
 	FMeshDescriptionBulkData()
@@ -1329,34 +1329,34 @@ public:
 	}
 
 	/** Serialization */
-	void Serialize(FArchive& Ar, UObject* Owner);
+	MESHDESCRIPTION_API void Serialize(FArchive& Ar, UObject* Owner);
 
 	/** Store a new mesh description in the bulk data */
-	void SaveMeshDescription(FMeshDescription& MeshDescription);
+	MESHDESCRIPTION_API void SaveMeshDescription(FMeshDescription& MeshDescription);
 
 	/** Load the mesh description from the bulk data */
-	void LoadMeshDescription(FMeshDescription& MeshDescription);
+	MESHDESCRIPTION_API void LoadMeshDescription(FMeshDescription& MeshDescription);
 
 	/** Empties the bulk data */
-	void Empty();
+	MESHDESCRIPTION_API void Empty();
 
 	/** Returns true if there is no bulk data available */
 	bool IsEmpty() const { return !BulkData.HasPayloadData(); }
 
 	/** Return unique ID string for this bulk data */
-	FString GetIdString() const;
+	MESHDESCRIPTION_API FString GetIdString() const;
 
 	/** Uses a hash as the GUID, useful to prevent recomputing content already in cache. */
-	void UseHashAsGuid();
+	MESHDESCRIPTION_API void UseHashAsGuid();
 
 	/** Gets the size of the serialized bulk data */
 	int64 GetBulkDataSize() { return BulkData.GetPayloadSize(); }
 
 	/** Returns all of the CustomVersions that can be used when serializing an FMeshDescription. */
-	static TConstArrayView<FGuid> GetMeshDescriptionCustomVersions();
+	static MESHDESCRIPTION_API TConstArrayView<FGuid> GetMeshDescriptionCustomVersions();
 
 private:
-	void UpdateMeshDescriptionFormat();
+	MESHDESCRIPTION_API void UpdateMeshDescriptionFormat();
 
 #if WITH_EDITOR
 	/** Protects simultaneous access to BulkData */
@@ -1369,7 +1369,7 @@ private:
 	FGuid Guid;
 
 	/** Gets this bulk data hash */
-	FGuid GetHash() const;
+	MESHDESCRIPTION_API FGuid GetHash() const;
 
 	/** Take a copy of the bulk data versioning so it can be propagated to the bulk data reader when deserializing MeshDescription */
 	FCustomVersionContainer CustomVersions;
