Media: Get selected tracks from media player if it supports its own track selection on open

#rnx

diff --git a/Engine/Plugins/Media/ElectraPlayer/Source/ElectraPlayerPlugin/Private/ElectraPlayerPlugin.cpp b/Engine/Plugins/Media/ElectraPlayer/Source/ElectraPlayerPlugin/Private/ElectraPlayerPlugin.cpp
--- a/Engine/Plugins/Media/ElectraPlayer/Source/ElectraPlayerPlugin/Private/ElectraPlayerPlugin.cpp
+++ b/Engine/Plugins/Media/ElectraPlayer/Source/ElectraPlayerPlugin/Private/ElectraPlayerPlugin.cpp
@@ -806,6 +806,8 @@ bool FElectraPlayerPlugin::GetPlayerFeatureFlag(EFeatureFlag flag) const
 			return true;
 		case EFeatureFlag::IsTrackSwitchSeamless:
 			return true;
+		case EFeatureFlag::PlayerSelectsDefaultTracks:
+			return true;
 		default:
 			break;
 	}
diff --git a/Engine/Source/Runtime/Media/Public/IMediaPlayer.h b/Engine/Source/Runtime/Media/Public/IMediaPlayer.h
--- a/Engine/Source/Runtime/Media/Public/IMediaPlayer.h
+++ b/Engine/Source/Runtime/Media/Public/IMediaPlayer.h
@@ -348,6 +348,7 @@ public:
 		AlwaysPullNewestVideoFrame,		//!< Mediaframework will not gate video frame output with its own timing, but assumes "ASAP" as output time for every sample
 		PlayerUsesInternalFlushOnSeek,	//!< The player implements an internal flush logic on seeks and Mediaframework will not issue an explicit Flush() call to it on seeks
 		IsTrackSwitchSeamless,			//!< If track switching is seamless then a flush of sinks is not necessary.
+		PlayerSelectsDefaultTracks,		//!< Whether or not the player selects suitable track defaults.
 	};
 	
 	virtual bool GetPlayerFeatureFlag(EFeatureFlag /*flag*/) const
diff --git a/Engine/Source/Runtime/MediaUtils/Private/MediaPlayerFacade.cpp b/Engine/Source/Runtime/MediaUtils/Private/MediaPlayerFacade.cpp
--- a/Engine/Source/Runtime/MediaUtils/Private/MediaPlayerFacade.cpp
+++ b/Engine/Source/Runtime/MediaUtils/Private/MediaPlayerFacade.cpp
@@ -1937,17 +1937,45 @@ void FMediaPlayerFacade::ResetTracks()
 
 void FMediaPlayerFacade::SelectDefaultTracks()
 {
-	FMediaPlayerTrackOptions TrackOptions;
-	if (ActivePlayerOptions.IsSet())
+	// See if the player has selected appropriate default tracks.
+	TSharedPtr<IMediaPlayer, ESPMode::ThreadSafe> CurrentPlayer(Player);
+	if (CurrentPlayer.IsValid() && CurrentPlayer->GetPlayerFeatureFlag(IMediaPlayer::EFeatureFlag::PlayerSelectsDefaultTracks))
 	{
-		TrackOptions = ActivePlayerOptions.GetValue().Tracks;
+		ResetTracks();
+		// Get what the player has selected as user defaults.
+		// The TrackSelection.PlayerSelection[...] will be updated in UpdateTrackSelectionWithPlayer()
+		// where the existence of sinks is checked for.
+		IMediaTracks& Tracks = CurrentPlayer->GetTracks();
+		for(int32 Idx=0; Idx<(int32)EMediaTrackType::Num; ++Idx)
+		{
+			TrackSelection.UserSelection[Idx] = Tracks.GetSelectedTrack((EMediaTrackType)Idx);
+		}
+		// If overrides are set, use them.
+		if (ActivePlayerOptions.IsSet())
+		{
+			FMediaPlayerTrackOptions TrackOptions;
+			TrackOptions = ActivePlayerOptions.GetValue().Tracks;
+			TrackSelection.UserSelection[(int32)EMediaTrackType::Audio] = TrackOptions.Audio;
+			TrackSelection.UserSelection[(int32)EMediaTrackType::Caption] = TrackOptions.Caption;
+			TrackSelection.UserSelection[(int32)EMediaTrackType::Metadata] = TrackOptions.Metadata;
+			TrackSelection.UserSelection[(int32)EMediaTrackType::Subtitle] = TrackOptions.Subtitle;
+			TrackSelection.UserSelection[(int32)EMediaTrackType::Video] = TrackOptions.Video;
+		}
 	}
+	else
+	{
+		FMediaPlayerTrackOptions TrackOptions;
+		if (ActivePlayerOptions.IsSet())
+		{
+			TrackOptions = ActivePlayerOptions.GetValue().Tracks;
+		}
 
-	TrackSelection.UserSelection[(int32)EMediaTrackType::Audio] = TrackOptions.Audio;
-	TrackSelection.UserSelection[(int32)EMediaTrackType::Caption] = TrackOptions.Caption;
-	TrackSelection.UserSelection[(int32)EMediaTrackType::Metadata] = TrackOptions.Metadata;
-	TrackSelection.UserSelection[(int32)EMediaTrackType::Subtitle] = TrackOptions.Subtitle;
-	TrackSelection.UserSelection[(int32)EMediaTrackType::Video] = TrackOptions.Video;
+		TrackSelection.UserSelection[(int32)EMediaTrackType::Audio] = TrackOptions.Audio;
+		TrackSelection.UserSelection[(int32)EMediaTrackType::Caption] = TrackOptions.Caption;
+		TrackSelection.UserSelection[(int32)EMediaTrackType::Metadata] = TrackOptions.Metadata;
+		TrackSelection.UserSelection[(int32)EMediaTrackType::Subtitle] = TrackOptions.Subtitle;
+		TrackSelection.UserSelection[(int32)EMediaTrackType::Video] = TrackOptions.Video;
+	}
 }
 
 
