Fixing up assert in MediaFramework facade code
#rnx

diff --git a/Engine/Source/Runtime/MediaUtils/Private/MediaPlayerFacade.cpp b/Engine/Source/Runtime/MediaUtils/Private/MediaPlayerFacade.cpp
--- a/Engine/Source/Runtime/MediaUtils/Private/MediaPlayerFacade.cpp
+++ b/Engine/Source/Runtime/MediaUtils/Private/MediaPlayerFacade.cpp
@@ -2920,8 +2920,8 @@ bool FMediaPlayerFacade::IsVideoSampleStillGood(const TRange<FMediaTimeStamp>& L
 			}
 
 			// So we must have a loop index change. Compute how many loops and change the range into one "unrolled" one as indicated by the playback direction...
-			int64 LoopIdxDiff = FMediaTimeStamp::GetSecondaryIndex(LowerSeqIdx)
-							  - FMediaTimeStamp::GetSecondaryIndex(UpperSeqIdx);
+			int32 LoopIdxDiff = FMediaTimeStamp::GetSecondaryIndex(UpperSeqIdx)
+							  - FMediaTimeStamp::GetSecondaryIndex(LowerSeqIdx);
 			check(LoopIdxDiff > 0);
 
 			double Duration = Player->GetControls().GetDuration().GetTotalSeconds();
