Fixes to MFW / Electra / ImgMedia / WmfMedia to address various playback timing / sync issues
#tests Various tests on local machine

diff --git a/Engine/Plugins/Media/ImgMedia/Source/ImgMedia/Internal/Loader/ImgMediaLoader.h b/Engine/Plugins/Media/ImgMedia/Source/ImgMedia/Internal/Loader/ImgMediaLoader.h
--- a/Engine/Plugins/Media/ImgMedia/Source/ImgMedia/Internal/Loader/ImgMediaLoader.h
+++ b/Engine/Plugins/Media/ImgMedia/Source/ImgMedia/Internal/Loader/ImgMediaLoader.h
@@ -355,7 +355,7 @@ public:
 	 * @param Frame The frame that was read, or nullptr if reading failed.
 	 * @param WorkTime How long to read this frame (in seconds).
 	 */
-	void NotifyWorkComplete(FImgMediaLoaderWork& CompletedWork, int32 FrameNumber,
+	void NotifyWorkComplete(FImgMediaLoaderWork& CompletedWork, int32 JobID, int32 FrameNumber,
 		const TSharedPtr<FImgMediaFrame, ESPMode::ThreadSafe>& Frame, float WorkTime);
 
 	/**
@@ -390,6 +390,15 @@ public:
 	bool IsFrameLast(const FTimespan& TimeStamp) const;
 
 protected:
+	struct FFrameNumberJobInfo
+	{
+		FFrameNumberJobInfo() = default;
+		FFrameNumberJobInfo(int32 InFrame, int32 InJobID) : Frame(InFrame), JobID(InJobID) {}
+
+		int32 Frame;
+		int32 JobID;
+		bool operator==(const FFrameNumberJobInfo& Other) const { return JobID == Other.JobID && Frame == Other.Frame; }
+	};
 
 	/**
 	 * Convert a collection of frame numbers to corresponding time ranges.
@@ -466,6 +475,15 @@ protected:
 	 */
 	uint32 TimeToFrameNumber(FTimespan Time) const;
 
+	/**
+	 * Get the frame number corresponding to the specified play head time without any checks for media range bounds.
+	 *
+	 * @param Time The play head time.
+	 * @return The corresponding frame number, or INDEX_NONE.
+	 * @see FrameNumberToTime
+	 */
+	int64 TimeToFrameNumberUnbound(FTimespan Time) const;
+
 	/**
 	 * Update the loader based on the current play position.
 	 *
@@ -632,7 +650,8 @@ private:
 	TArray<int32> PendingFrameNumbers;
 
 	/** Collection of frame numbers that are being read. */
-	TArray<int32> QueuedFrameNumbers;
+	TArray<FFrameNumberJobInfo> QueuedFrameNumbers;
+	int32 NextFrameNumberJobID;
 
 	/** Object pool for reusable work items. */
 	TArray<FImgMediaLoaderWork*> WorkPool;
diff --git a/Engine/Plugins/Media/ImgMedia/Source/ImgMedia/Internal/Player/ImgMediaPlayer.h b/Engine/Plugins/Media/ImgMedia/Source/ImgMedia/Internal/Player/ImgMediaPlayer.h
--- a/Engine/Plugins/Media/ImgMedia/Source/ImgMedia/Internal/Player/ImgMediaPlayer.h
+++ b/Engine/Plugins/Media/ImgMedia/Source/ImgMedia/Internal/Player/ImgMediaPlayer.h
@@ -108,7 +108,7 @@ protected:
 	virtual bool FetchVideo(TRange<FTimespan> TimeRange, TSharedPtr<IMediaTextureSample, ESPMode::ThreadSafe>& OutSample) override;
 	virtual void FlushSamples() override;
 
-	virtual EFetchBestSampleResult FetchBestVideoSampleForTimeRange(const TRange<FMediaTimeStamp> & TimeRange, TSharedPtr<IMediaTextureSample, ESPMode::ThreadSafe>& OutSample, bool bReverse) override;
+	virtual EFetchBestSampleResult FetchBestVideoSampleForTimeRange(const TRange<FMediaTimeStamp>& TimeRange, TSharedPtr<IMediaTextureSample, ESPMode::ThreadSafe>& OutSample, bool bReverse, bool bConsistentResult) override;
 	virtual bool PeekVideoSampleTime(FMediaTimeStamp & TimeStamp) override;
 	virtual bool DiscardVideoSamples(const TRange<FMediaTimeStamp>& TimeRange, bool bReverse) override;
 
diff --git a/Engine/Plugins/Media/ImgMedia/Source/ImgMedia/Private/Loader/ImgMediaLoader.cpp b/Engine/Plugins/Media/ImgMedia/Source/ImgMedia/Private/Loader/ImgMediaLoader.cpp
--- a/Engine/Plugins/Media/ImgMedia/Source/ImgMedia/Private/Loader/ImgMediaLoader.cpp
+++ b/Engine/Plugins/Media/ImgMedia/Source/ImgMedia/Private/Loader/ImgMediaLoader.cpp
@@ -174,6 +174,7 @@ FImgMediaLoader::FImgMediaLoader(const TSharedRef<FImgMediaScheduler, ESPMode::T
 	, SequenceFrameRate(0, 0)
 	, LastRequestedFrame(INDEX_NONE)
 	, RetryCount(0)
+	, NextFrameNumberJobID(0)
 	, UseGlobalCache(false)
 	, SmartCacheSettings(InSmartCacheSettings)
 	, bIsPlaybackBlocking(false)
@@ -229,7 +230,13 @@ uint64 FImgMediaLoader::GetBitRate() const
 void FImgMediaLoader::GetBusyTimeRanges(TRangeSet<FTimespan>& OutRangeSet) const
 {
 	FScopeLock Lock(&CriticalSection);
-	FrameNumbersToTimeRanges(QueuedFrameNumbers, OutRangeSet);
+	TArray<int32> FrameNumbers;
+	FrameNumbers.Reserve(QueuedFrameNumbers.Num());
+	for (const auto& Entry : QueuedFrameNumbers)
+	{
+		FrameNumbers.Push(Entry.Frame);
+	}
+	FrameNumbersToTimeRanges(FrameNumbers, OutRangeSet);
 }
 
 
@@ -333,7 +340,7 @@ void FImgMediaLoader::Flush()
 void FImgMediaLoader::Seek(const FMediaTimeStamp& SeekTarget, bool bReverse)
 {
 	ResetFetchLogic();
-	QueuedSampleFetch.LastTimeStamp = { FrameNumberToTime(TimeToFrameNumber(SeekTarget.Time)), SeekTarget.SequenceIndex };
+	QueuedSampleFetch.LastTimeStamp = { FMediaTimeStamp(FrameNumberToTime(TimeToFrameNumber(SeekTarget.Time)), SeekTarget.SequenceIndex) };
 	QueuedSampleFetch.LastDuration = FTimespan::Zero();
 }
 
@@ -428,214 +435,180 @@ IMediaSamples::EFetchBestSampleResult FImgMediaLoader::FetchBestVideoSampleForTi
 			EndLoopIndex = FMediaTimeStamp::GetSecondaryIndex(TimeRange.GetUpperBoundValue().SequenceIndex) - 1;
 		}
 
-		if (bIsLoopingEnabled)
-		{
-			// Modulo with sequence duration to take care of looping.
-			StartTime = ModuloTime(StartTime);
-			EndTime = ModuloTime(EndTime);
-		}
-
-		// Get start and end frame indices for this time range.
-		int32 StartIndex = TimeToFrameNumber(StartTime);
-		int32 EndIndex = TimeToFrameNumber(EndTime);
+		int32 NumFrames = GetNumImages();
+		FTimespan FrameLength = FTimespan::FromSeconds(SequenceFrameRate.AsInterval());
 
-		// Sanity checks on returned indices...
-		if ((uint32)StartIndex == INDEX_NONE && (uint32)EndIndex == INDEX_NONE)
-		{
-			return IMediaSamples::EFetchBestSampleResult::NoSample;
-		}
-
-		if ((uint32)StartIndex == INDEX_NONE)
-		{
-			StartIndex = 0;
-		}
-		else if ((uint32)EndIndex == INDEX_NONE)
-		{
-			EndIndex = GetNumImages() - 1;
-		}
+		FTimespan UnwrappedStartTime = StartTime + FrameLength * NumFrames * StartLoopIndex;
+		FTimespan UnwrappedEndTime = EndTime + FrameLength * NumFrames * EndLoopIndex;
+		int64 UnwrappedStartIndex = TimeToFrameNumberUnbound(UnwrappedStartTime);
+		int64 UnwrappedEndIndex = TimeToFrameNumberUnbound(UnwrappedEndTime);
+		FTimespan UnwrappedStartTimeAligned = UnwrappedStartIndex * FrameLength;
+		FTimespan UnwrappedEndTimeAligned = UnwrappedEndIndex * FrameLength;
 
-		// Find the frame that overlaps the most with the given range & is furthest along on the timeline
-		int32 MaxIdx = -1;
-		int32 MaxIdxLoopIndexOffset = 0;
+		int64 MaxIdx;
 		if (PlayRate >= 0.0f)
 		{
 			// Forward...
 
-			// Per default we select the first frame in range
-			MaxIdx = StartIndex;
-
 			// More than one possibility?
-			if (EndIndex != StartIndex || StartLoopIndex != EndLoopIndex)
+			if (UnwrappedStartIndex != UnwrappedEndIndex)
 			{
-				FTimespan FrameLength = FTimespan::FromSeconds(SequenceFrameRate.AsInterval());
-
-				// Compute visible duration of first frame...
-				FTimespan FirstTime = FrameNumberToTime(StartIndex);
-				FTimespan FirstFrameDuration = FrameLength - FTimespan(FMath::Max((TimeRange.GetLowerBoundValue().Time - FirstTime).GetTicks(), 0));
+				FTimespan FirstFrameDuration = FrameLength - (UnwrappedStartTime - UnwrappedStartTimeAligned);
+				FTimespan LastFrameDuration = UnwrappedEndTime - UnwrappedEndTimeAligned;
 
-				// Compute the length of the last frame
-				FTimespan LastTime = FrameNumberToTime(EndIndex);
-				FTimespan LastFrameDuration = FrameLength - FTimespan(FMath::Max((LastTime - TimeRange.GetUpperBoundValue().Time).GetTicks(), 0));
-
-				// First one wins by default
-				FTimespan MaxFrameDuration = FirstFrameDuration;
-
-				int32 NumFrames = GetNumImages();
-
-				// Do we have a "second to last" frame? (it would always have full coverage)
-				int64 EndIndexUnlooped = (EndLoopIndex - StartLoopIndex) * NumFrames + EndIndex;
-				if ((EndIndexUnlooped - StartIndex) > 1)
+				FTimespan MaxFrameDuration;
+				if ((UnwrappedEndIndex - UnwrappedStartIndex) > 1)
 				{
-					// We do. Record it as best so far instead of the first one...
-					MaxIdx = (EndIndex > 0) ? (EndIndex - 1) : (NumFrames - 1);
+					MaxIdx = UnwrappedEndIndex - 1;
 					MaxFrameDuration = FrameLength;
 				}
-
-				// Is the last one even better?
+				else
+				{
+					MaxIdx = UnwrappedStartIndex;
+					MaxFrameDuration = FirstFrameDuration;
+				}
 				if (LastFrameDuration >= MaxFrameDuration)
 				{
-					MaxIdx = EndIndex;
+					MaxIdx = UnwrappedEndIndex;
+					MaxFrameDuration = LastFrameDuration;
 				}
 			}
+			else
+			{
+				MaxIdx = UnwrappedEndIndex;
+			}
 		}
 		else
 		{
-			// Backwards...
-
-			// Per default we select the last frame in range
-			MaxIdx = EndIndex;
+			// Backward...
 
 			// More than one possibility?
-			if (EndIndex != StartIndex || StartLoopIndex != EndLoopIndex)
+			if (UnwrappedStartIndex != UnwrappedEndIndex)
 			{
-				FTimespan FrameLength = FTimespan::FromSeconds(SequenceFrameRate.AsInterval());
+				FTimespan FirstFrameDuration = FrameLength - (UnwrappedStartTime - UnwrappedStartTimeAligned);
+				FTimespan LastFrameDuration = UnwrappedEndTime - UnwrappedEndTimeAligned;
 
-				// Compute the length of the last frame
-				FTimespan FirstTime = FrameNumberToTime(EndIndex);
-				FTimespan FirstFrameDuration = FrameLength - FTimespan(FMath::Max((FirstTime - TimeRange.GetUpperBoundValue().Time).GetTicks(), 0));
-
-				// Compute visible duration of first frame...
-				FTimespan LastTime = FrameNumberToTime(StartIndex);
-				FTimespan LastFrameDuration = FrameLength - FTimespan(FMath::Max((TimeRange.GetLowerBoundValue().Time - LastTime).GetTicks(), 0));
-
-				// Last one wins by default
-				FTimespan MaxFrameDuration = LastFrameDuration;
-
-				int32 NumFrames = GetNumImages();
-
-				// Do we have a "second to last" frame? (it would always have full coverage)
-				int64 EndIndexUnlooped = (EndLoopIndex - StartLoopIndex) * NumFrames + EndIndex;
-				if ((EndIndexUnlooped - StartIndex) > 1)
+				FTimespan MaxFrameDuration;
+				if ((UnwrappedEndIndex - UnwrappedStartIndex) > 1)
 				{
-					// We do. Record it as best so far instead of the first one...
-					MaxIdx = (StartIndex < (NumFrames - 1)) ? (StartIndex + 1) : 0;
+					MaxIdx = UnwrappedStartIndex + 1;
 					MaxFrameDuration = FrameLength;
 				}
-
-				// Is the first one even better?
+				else
+				{
+					MaxIdx = UnwrappedEndIndex;
+					MaxFrameDuration = LastFrameDuration;
+				}
 				if (FirstFrameDuration >= MaxFrameDuration)
 				{
-					MaxIdx = StartIndex;
+					MaxIdx = UnwrappedStartIndex;
+					MaxFrameDuration = FirstFrameDuration;
 				}
 			}
+			else
+			{
+				MaxIdx = UnwrappedStartIndex;
+			}
 		}
 
 		// - - - - - - - - - - - - - - - - - - - - - - - - - - - - - - - - - - - - - - - - - - - - - - - - - - - -
 
 		// Anything?
-		if (MaxIdx >= 0)
+		int32 MaxLoopIdx;
+		if (bIsLoopingEnabled)
 		{
-			const TSharedPtr<FImgMediaFrame, ESPMode::ThreadSafe>* Frame;
+			MaxLoopIdx = int32(MaxIdx / NumFrames);
+			bool bNegative = (MaxIdx < 0);
+			MaxIdx = bNegative ? -MaxIdx : MaxIdx;
+			MaxIdx = MaxIdx - (MaxIdx / NumFrames) * NumFrames;
+			if (bNegative)
+			{
+				MaxIdx = (NumFrames - 1) - MaxIdx;
+			}
+		}
+		else
+		{
+			MaxLoopIdx = 0;	// ignored below
+			// Clamp this between [0..N-1], but keep a look out for any loop index that might persist from an earlier looping state
+			MaxIdx = FMath::Clamp(MaxIdx, StartLoopIndex * NumFrames, StartLoopIndex * NumFrames + NumFrames - 1);
+		}
 
-			// Request data for the frame we would like... (in case it's not requested or in, yet)
-			RequestFrame(MaxIdx, PlayRate, bIsLoopingEnabled);
+		const TSharedPtr<FImgMediaFrame, ESPMode::ThreadSafe>* Frame;
 
-			FScopeLock Lock(&CriticalSection);
+		// Request data for the frame we would like... (in case it's not requested or in, yet)
+		RequestFrame(MaxIdx, PlayRate, bIsLoopingEnabled);
 
-			// Get a frame if we have one available right now...
-			Frame = GetCachedFrame(MaxIdx);
+		FScopeLock Lock(&CriticalSection);
 
-			// Got a potential frame?
-			if (Frame)
-			{
-				// Yes.
+		// Get a frame if we have one available right now...
+		Frame = GetCachedFrame(MaxIdx);
 
-				/*
-					The player facade will pass in a time range with current primary and secondary sequence indices.
-					We need to make sure that we detect if we return a value "before" the last one (loop case) and
-					adjust the secondary index accordingly.
-				*/
+		// Got a potential frame?
+		if (Frame)
+		{
+			// Yes.
 
-				RetryCount = 0;
+			/*
+				The player facade will pass in a time range with current primary and secondary sequence indices.
+				We need to make sure that we detect if we return a value "before" the last one (loop case) and
+				adjust the secondary index accordingly.
+			*/
 
-				FTimespan SampleTime = FrameNumberToTime(MaxIdx);
+			RetryCount = 0;
 
-				// Check if we wrapped around the end of the media vs. the last frame we returned...
-				// (skip this if we are not looping or if this is either the first frame ever or first frame after a seek)
-				if (bIsLoopingEnabled && !QueuedSampleFetch.LastDuration.IsZero())
-				{
-					if (PlayRate >= 0.0f)
-					{
-						if (QueuedSampleFetch.LastTimeStamp.Time >= SampleTime)
-						{
-							++QueuedSampleFetch.LoopIndex;
-						}
-					}
-					else
-					{
-						if (QueuedSampleFetch.LastTimeStamp.Time <= SampleTime)
-						{
-							--QueuedSampleFetch.LoopIndex;
-						}
-					}
-				}
+			FTimespan SampleTime = FrameNumberToTime(MaxIdx);
 
-				// Track state & setup sample for caller...
-				QueuedSampleFetch.LastTimeStamp = FMediaTimeStamp(SampleTime, FMediaTimeStamp::MakeSequenceIndex(FMediaTimeStamp::GetPrimaryIndex(TimeRange.GetLowerBoundValue().SequenceIndex), QueuedSampleFetch.LoopIndex));
-				QueuedSampleFetch.LastDuration = FTimespan::FromSeconds(Frame->Get()->GetFrameRate().AsInterval());
+			// Set loop index if looping is active (otherwise we keep what was known before)
+			if (bIsLoopingEnabled && !QueuedSampleFetch.LastDuration.IsZero())
+			{
+				QueuedSampleFetch.LoopIndex = MaxLoopIdx;
+			}
 
-				// We are clear to return it as new result... Make a sample & initialize it...
-				auto Sample = MakeShared<FImgMediaTextureSample, ESPMode::ThreadSafe>();
+			// Track state & setup sample for caller...
+			QueuedSampleFetch.LastTimeStamp = FMediaTimeStamp(SampleTime, FMediaTimeStamp::MakeSequenceIndex(FMediaTimeStamp::GetPrimaryIndex(TimeRange.GetLowerBoundValue().SequenceIndex), QueuedSampleFetch.LoopIndex));
+			QueuedSampleFetch.LastDuration = FTimespan::FromSeconds(Frame->Get()->GetFrameRate().AsInterval());
 
-				ImgMediaLoader::CheckAndUpdateImgDimensions(SequenceDim, Frame->Get()->GetDim());
+			// We are clear to return it as new result... Make a sample & initialize it...
+			auto Sample = MakeShared<FImgMediaTextureSample, ESPMode::ThreadSafe>();
 
-				if (Sample->Initialize(*Frame->Get(), SequenceDim, QueuedSampleFetch.LastTimeStamp, QueuedSampleFetch.LastDuration, GetNumMipLevels(), TilingDescription))
-				{
-					OutSample = Sample;
-					CSV_EVENT(ImgMedia, TEXT("LoaderFetchHit %d %d-%d"), MaxIdx, StartIndex, EndIndex);
-					return IMediaSamples::EFetchBestSampleResult::Ok;
-				}
+			ImgMediaLoader::CheckAndUpdateImgDimensions(SequenceDim, Frame->Get()->GetDim());
+
+			if (Sample->Initialize(*Frame->Get(), SequenceDim, QueuedSampleFetch.LastTimeStamp, QueuedSampleFetch.LastDuration, GetNumMipLevels(), TilingDescription))
+			{
+				OutSample = Sample;
+				CSV_EVENT(ImgMedia, TEXT("LoaderFetchHit %d %d-%d"), MaxIdx, UnwrappedStartIndex % NumFrames, UnwrappedEndIndex % NumFrames);
+				return IMediaSamples::EFetchBestSampleResult::Ok;
 			}
-			else
+		}
+		else
+		{
+			// We did not get a frame...
+			// Could we have lost a frame that we previously loaded
+			// due to the global cache being full?
+			if (UseGlobalCache)
 			{
-				// We did not get a frame...
-				// Could we have lost a frame that we previously loaded
-				// due to the global cache being full?
-				if (UseGlobalCache)
+				// Did we get a frame previously?
+				if (LastRequestedFrame != INDEX_NONE)
 				{
-					// Did we get a frame previously?
-					if (LastRequestedFrame != INDEX_NONE)
+					// Are we loading any frames?
+					if ((PendingFrameNumbers.Num() == 0) && (QueuedFrameNumbers.Num() == 0))
 					{
-						// Are we loading any frames?
-						if ((PendingFrameNumbers.Num() == 0) && (QueuedFrameNumbers.Num() == 0))
+						// Nope...
+						// Wait for this to happen for one more frame.
+						// If we have a 1 image sequence, we might have just missed the frame
+						// so try again next time.
+						RetryCount++;
+						if (RetryCount > 1)
 						{
-							// Nope...
-							// Wait for this to happen for one more frame.
-							// If we have a 1 image sequence, we might have just missed the frame
-							// so try again next time.
-							RetryCount++;
-							if (RetryCount > 1)
-							{
-								UE_LOG(LogImgMedia, Error, TEXT("Reloading frames. The global cache may be too small."));
-								LastRequestedFrame = INDEX_NONE;
-								RetryCount = 0;
-							}
+							UE_LOG(LogImgMedia, Error, TEXT("Reloading frames. The global cache may be too small."));
+							LastRequestedFrame = INDEX_NONE;
+							RetryCount = 0;
 						}
 					}
 				}
 			}
 		}
-		CSV_EVENT(ImgMedia, TEXT("LoaderFetchMiss %d-%d"), StartIndex, EndIndex);
+		CSV_EVENT(ImgMedia, TEXT("LoaderFetchMiss %d-%d"), UnwrappedStartIndex % NumFrames, UnwrappedEndIndex % NumFrames);
 	}
 	return IMediaSamples::EFetchBestSampleResult::NoSample;
 }
@@ -651,9 +624,9 @@ bool FImgMediaLoader::PeekVideoSampleTime(FMediaTimeStamp &TimeStamp, bool bIsLo
 		if (PlayRate >= 0.0f)
 		{
 			TimeStamp = QueuedSampleFetch.LastTimeStamp + QueuedSampleFetch.LastDuration;
-			if (TimeStamp >= SequenceDuration)
+			if (TimeStamp.Time >= SequenceDuration)
 			{
-				TimeStamp -= SequenceDuration;
+				TimeStamp.Time -= SequenceDuration;
 				TimeStamp.SequenceIndex = FMediaTimeStamp::AdjustSecondaryIndex(QueuedSampleFetch.LastTimeStamp.SequenceIndex, 1);
 			}
 		}
@@ -662,7 +635,7 @@ bool FImgMediaLoader::PeekVideoSampleTime(FMediaTimeStamp &TimeStamp, bool bIsLo
 			TimeStamp = QueuedSampleFetch.LastTimeStamp - QueuedSampleFetch.LastDuration;
 			if (TimeStamp.Time.GetTicks() < 0)
 			{
-				TimeStamp += SequenceDuration;
+				TimeStamp.Time += SequenceDuration;
 				TimeStamp.SequenceIndex = FMediaTimeStamp::AdjustSecondaryIndex(QueuedSampleFetch.LastTimeStamp.SequenceIndex, -1);
 			}
 		}
@@ -729,8 +702,8 @@ IQueuedWork* FImgMediaLoader::GetWork()
 	}
 	
 	// Set up work.
-	Work->Initialize(FrameNumber, DesiredMipsAndTiles, ExistingFrame);
-	QueuedFrameNumbers.Add(FrameNumber);
+	Work->Initialize(NextFrameNumberJobID, FrameNumber, DesiredMipsAndTiles, ExistingFrame);
+	QueuedFrameNumbers.Emplace(FrameNumber, NextFrameNumberJobID++);
 
 	return Work;
 }
@@ -797,7 +770,7 @@ void FImgMediaLoader::FrameNumbersToTimeRanges(const TArray<int32>& FrameNumbers
 		return;
 	}
 
-	for (const auto FrameNumber : FrameNumbers)
+	for (const auto& FrameNumber : FrameNumbers)
 	{
 		const FTimespan FrameStartTime = FrameNumberToTime(FrameNumber);
 		const FTimespan NextStartTime = FrameNumberToTime(FrameNumber + 1);
@@ -1179,6 +1152,13 @@ FTimespan FImgMediaLoader::FrameNumberToTime(uint32 FrameNumber) const
 }
 
 
+int64 FImgMediaLoader::TimeToFrameNumberUnbound(FTimespan Time) const
+{
+	// note: do NOT assume 100% reversability between TimeToFrameNumber() and FrameNumberToTime()! The FTimespan "only" has 100ns resolution and does some rounding when converting from seconds -> this migth shift things around!
+	return static_cast<int64>(FMath::FloorToInt(Time.GetTotalSeconds() * (static_cast<double>(SequenceFrameRate.Numerator) / SequenceFrameRate.Denominator)));
+}
+
+
 void FImgMediaLoader::Update(int32 PlayHeadFrame, float PlayRate, bool Loop)
 {
 	// In case reader needs to do something once per frame.
@@ -1294,7 +1274,7 @@ void FImgMediaLoader::Update(int32 PlayHeadFrame, float PlayRate, bool Loop)
 	// determine queued frame numbers that can be discarded
 	for (int32 Idx = QueuedFrameNumbers.Num() - 1; Idx >= 0; --Idx)
 	{
-		const int32 FrameNumber = QueuedFrameNumbers[Idx];
+		const int32 FrameNumber = QueuedFrameNumbers[Idx].Frame;
 
 		if (!FramesToLoad.Contains(FrameNumber))
 		{
@@ -1347,16 +1327,29 @@ void FImgMediaLoader::Update(int32 PlayHeadFrame, float PlayRate, bool Loop)
 			}
 		}
 		
-		if ((NeedFrame) && !QueuedFrameNumbers.Contains(FrameNumber))
+		if (NeedFrame)
 		{
-			// Do we actually have a frame for this?
-			if (ImagePaths[0][FrameNumber].Len() == 0)
+			bool FoundFrame = false;
+			for (const auto& Entry : QueuedFrameNumbers)
 			{
-				TryAddEmptyFrame(FrameNumber);
+				if (Entry.Frame == FrameNumber)
+				{
+					FoundFrame = true;
+					break;
+				}
 			}
-			else
+
+			if (!FoundFrame)
 			{
-				PendingFrameNumbers.Add(FrameNumber);
+				// Do we actually have a frame for this?
+				if (ImagePaths[0][FrameNumber].Len() == 0)
+				{
+					TryAddEmptyFrame(FrameNumber);
+				}
+				else
+				{
+					PendingFrameNumbers.Add(FrameNumber);
+				}
 			}
 		}
 	}
@@ -1415,13 +1408,13 @@ bool FImgMediaLoader::TryAddEmptyFrame(int32 FrameNumber)
 	return bAddFrame;
 }
 
-void FImgMediaLoader::NotifyWorkComplete(FImgMediaLoaderWork& CompletedWork, int32 FrameNumber,
+void FImgMediaLoader::NotifyWorkComplete(FImgMediaLoaderWork& CompletedWork, int32 JobID, int32 FrameNumber,
 	const TSharedPtr<FImgMediaFrame, ESPMode::ThreadSafe>& Frame, float WorkTime)
 {
 	FScopeLock Lock(&CriticalSection);
 
 	// if frame is still needed, add it to the cache
-	if (QueuedFrameNumbers.Remove(FrameNumber) > 0)
+	if (QueuedFrameNumbers.Remove({ FrameNumber, JobID }) > 0)
 	{
 		AddFrameToCache(FrameNumber, Frame);
 
diff --git a/Engine/Plugins/Media/ImgMedia/Source/ImgMedia/Private/Loader/ImgMediaLoaderWork.cpp b/Engine/Plugins/Media/ImgMedia/Source/ImgMedia/Private/Loader/ImgMediaLoaderWork.cpp
--- a/Engine/Plugins/Media/ImgMedia/Source/ImgMedia/Private/Loader/ImgMediaLoaderWork.cpp
+++ b/Engine/Plugins/Media/ImgMedia/Source/ImgMedia/Private/Loader/ImgMediaLoaderWork.cpp
@@ -32,10 +32,11 @@ FImgMediaLoaderWork::FImgMediaLoaderWork(const TSharedRef<FImgMediaLoader, ESPMo
 /* FImgMediaLoaderWork interface
  *****************************************************************************/
 
-void FImgMediaLoaderWork::Initialize(int32 InFrameNumber,
+void FImgMediaLoaderWork::Initialize(int32 InJobID, int32 InFrameNumber,
 	TMap<int32, FImgMediaTileSelection> InMipTiles,
 	TSharedPtr<FImgMediaFrame, ESPMode::ThreadSafe> InExistingFrame)
 {
+	JobID = InJobID;
 	FrameNumber = InFrameNumber;
 	MipTiles = MoveTemp(InMipTiles);
 	ExistingFrame = InExistingFrame;
@@ -114,7 +115,7 @@ void FImgMediaLoaderWork::Finalize(TSharedPtr<FImgMediaFrame, ESPMode::ThreadSaf
 
 	if (Owner.IsValid())
 	{
-		Owner->NotifyWorkComplete(*this, FrameNumber, Frame, WorkTime);
+		Owner->NotifyWorkComplete(*this, JobID, FrameNumber, Frame, WorkTime);
 	}
 	else
 	{
diff --git a/Engine/Plugins/Media/ImgMedia/Source/ImgMedia/Private/Loader/ImgMediaLoaderWork.h b/Engine/Plugins/Media/ImgMedia/Source/ImgMedia/Private/Loader/ImgMediaLoaderWork.h
--- a/Engine/Plugins/Media/ImgMedia/Source/ImgMedia/Private/Loader/ImgMediaLoaderWork.h
+++ b/Engine/Plugins/Media/ImgMedia/Source/ImgMedia/Private/Loader/ImgMediaLoaderWork.h
@@ -33,13 +33,14 @@ public:
 	/**
 	 * Initialize this work item.
 	 *
+	 * @param InJobID An ID to identify the job while in flight
 	 * @param InFrameNumber The number of the image frame.
 	 * @param InMipLevel Will read in this level and all higher levels.
 	 * @param InTileSelection Which tiles to read.
 	 * @param InExistingFrame If set, then use this frame to load into.
 	 * @see Shutdown
 	 */
-	void Initialize(int32 InFrameNumber, TMap<int32, FImgMediaTileSelection> InMipTiles, TSharedPtr<FImgMediaFrame, ESPMode::ThreadSafe> InExistingFrame);
+	void Initialize(int32 InJobID, int32 InFrameNumber, TMap<int32, FImgMediaTileSelection> InMipTiles, TSharedPtr<FImgMediaFrame, ESPMode::ThreadSafe> InExistingFrame);
 
 public:
 
@@ -60,6 +61,9 @@ protected:
 
 private:
 
+	/** ID to identify this workload instance when in flight */
+	int32 JobID;
+
 	/** The number of the image frame. */
 	int32 FrameNumber;
 
diff --git a/Engine/Plugins/Media/ImgMedia/Source/ImgMedia/Private/Player/ImgMediaPlayer.cpp b/Engine/Plugins/Media/ImgMedia/Source/ImgMedia/Private/Player/ImgMediaPlayer.cpp
--- a/Engine/Plugins/Media/ImgMedia/Source/ImgMedia/Private/Player/ImgMediaPlayer.cpp
+++ b/Engine/Plugins/Media/ImgMedia/Source/ImgMedia/Private/Player/ImgMediaPlayer.cpp
@@ -665,8 +665,10 @@ bool FImgMediaPlayer::DiscardVideoSamples(const TRange<FMediaTimeStamp>& TimeRan
 }
 
 
-IMediaSamples::EFetchBestSampleResult FImgMediaPlayer::FetchBestVideoSampleForTimeRange(const TRange<FMediaTimeStamp>& TimeRange, TSharedPtr<IMediaTextureSample, ESPMode::ThreadSafe>& OutSample, bool bReverse)
+IMediaSamples::EFetchBestSampleResult FImgMediaPlayer::FetchBestVideoSampleForTimeRange(const TRange<FMediaTimeStamp>& TimeRange, TSharedPtr<IMediaTextureSample, ESPMode::ThreadSafe>& OutSample, bool bReverse, bool bConsistentResult)
 {
+	// note: the results produced by this player are "consistent" in respect to which frame is returned for a specific range, so we just disregard the bConsistentResult flag
+
 	IMediaSamples::EFetchBestSampleResult SampleResult = EFetchBestSampleResult::NoSample;
 
 	if (Loader.IsValid() && IsInitialized())
diff --git a/Engine/Plugins/Media/ImgMedia/Source/ImgMedia/Private/Scheduler/ImgMediaScheduler.cpp b/Engine/Plugins/Media/ImgMedia/Source/ImgMedia/Private/Scheduler/ImgMediaScheduler.cpp
--- a/Engine/Plugins/Media/ImgMedia/Source/ImgMedia/Private/Scheduler/ImgMediaScheduler.cpp
+++ b/Engine/Plugins/Media/ImgMedia/Source/ImgMedia/Private/Scheduler/ImgMediaScheduler.cpp
@@ -186,7 +186,7 @@ IQueuedWork* FImgMediaScheduler::GetWorkOrReturnToPool(FImgMediaSchedulerThread*
 
 				while (CheckedLoaders++ < NumLoaders)
 				{
-					LoaderRoundRobin = (LoaderRoundRobin + 1) % NumLoaders;
+					LoaderRoundRobin = (LoaderRoundRobin + 1) % Loaders.Num();
 
 					TSharedPtr<FImgMediaLoader, ESPMode::ThreadSafe> Loader = Loaders[LoaderRoundRobin].Pin();
 
diff --git a/Engine/Plugins/Media/WebMMedia/Source/WebMMedia/Private/Player/WebMMediaPlayer.cpp b/Engine/Plugins/Media/WebMMedia/Source/WebMMedia/Private/Player/WebMMediaPlayer.cpp
--- a/Engine/Plugins/Media/WebMMedia/Source/WebMMedia/Private/Player/WebMMediaPlayer.cpp
+++ b/Engine/Plugins/Media/WebMMedia/Source/WebMMedia/Private/Player/WebMMediaPlayer.cpp
@@ -768,7 +768,7 @@ void FWebMMediaPlayer::Stop()
 
 void FWebMMediaPlayer::AddVideoSampleFromDecodingThread(TSharedRef<FWebMMediaTextureSample, ESPMode::ThreadSafe> Sample)
 {
-	if (Sample->GetTime() < CurrentTime)
+	if (Sample->GetTime().Time < CurrentTime)
 	{
 		// We don't care about expired samples
 		return;
@@ -779,7 +779,7 @@ void FWebMMediaPlayer::AddVideoSampleFromDecodingThread(TSharedRef<FWebMMediaTex
 
 void FWebMMediaPlayer::AddAudioSampleFromDecodingThread(TSharedRef<FWebMMediaAudioSample, ESPMode::ThreadSafe> Sample)
 {
-	if (Sample->GetTime() < CurrentTime)
+	if (Sample->GetTime().Time < CurrentTime)
 	{
 		// We don't care about expired samples
 		return;
diff --git a/Engine/Plugins/Media/WmfMedia/Source/WmfMedia/Private/Player/WmfMediaAudioSample.h b/Engine/Plugins/Media/WmfMedia/Source/WmfMedia/Private/Player/WmfMediaAudioSample.h
--- a/Engine/Plugins/Media/WmfMedia/Source/WmfMedia/Private/Player/WmfMediaAudioSample.h
+++ b/Engine/Plugins/Media/WmfMedia/Source/WmfMedia/Private/Player/WmfMediaAudioSample.h
@@ -25,7 +25,7 @@ public:
 		: Channels(0)
 		, Duration(FTimespan::Zero())
 		, SampleRate(0)
-		, Time(FTimespan::Zero())
+		, Time(FTimespan::Zero(), 0)
 	{ }
 
 	/** Virtual destructor. */
@@ -46,7 +46,7 @@ public:
 		uint32 InSize,
 		uint32 InChannels,
 		uint32 InSampleRate,
-		FTimespan InTime,
+		FMediaTimeStamp InTime,
 		FTimespan InDuration)
 	{
 		if ((InBuffer == nullptr) || (InSize == 0))
@@ -101,7 +101,7 @@ public:
 
 	virtual FMediaTimeStamp GetTime() const override
 	{
-		return FMediaTimeStamp(Time);
+		return Time;
 	}
 
 private:
@@ -119,7 +119,7 @@ private:
 	uint32 SampleRate;
 
 	/** Presentation time for which the sample was generated. */
-	FTimespan Time;
+	FMediaTimeStamp Time;
 };
 
 
diff --git a/Engine/Plugins/Media/WmfMedia/Source/WmfMedia/Private/Player/WmfMediaBinarySample.h b/Engine/Plugins/Media/WmfMedia/Source/WmfMedia/Private/Player/WmfMediaBinarySample.h
--- a/Engine/Plugins/Media/WmfMedia/Source/WmfMedia/Private/Player/WmfMediaBinarySample.h
+++ b/Engine/Plugins/Media/WmfMedia/Source/WmfMedia/Private/Player/WmfMediaBinarySample.h
@@ -39,7 +39,7 @@ public:
 	bool Initialize(
 		const void* InBuffer,
 		uint32 InSize,
-		FTimespan InTime,
+		FMediaTimeStamp InTime,
 		FTimespan InDuration)
 	{
 		if ((InBuffer == nullptr) || (InSize == 0))
@@ -77,7 +77,7 @@ public:
 
 	virtual FMediaTimeStamp GetTime() const override
 	{
-		return FMediaTimeStamp(Time);
+		return Time;
 	}
 
 private:
@@ -89,5 +89,5 @@ private:
 	FTimespan Duration;
 
 	/** Presentation time for which the sample was generated. */
-	FTimespan Time;
+	FMediaTimeStamp Time;
 };
diff --git a/Engine/Plugins/Media/WmfMedia/Source/WmfMedia/Private/Player/WmfMediaOverlaySample.h b/Engine/Plugins/Media/WmfMedia/Source/WmfMedia/Private/Player/WmfMediaOverlaySample.h
--- a/Engine/Plugins/Media/WmfMedia/Source/WmfMedia/Private/Player/WmfMediaOverlaySample.h
+++ b/Engine/Plugins/Media/WmfMedia/Source/WmfMedia/Private/Player/WmfMediaOverlaySample.h
@@ -37,7 +37,7 @@ public:
 	 */
 	bool Initialize(
 		const char* InBuffer,
-		FTimespan InTime,
+		FMediaTimeStamp InTime,
 		FTimespan InDuration)
 	{
 		if (InBuffer == nullptr)
@@ -99,7 +99,7 @@ public:
 
 	virtual FMediaTimeStamp GetTime() const override
 	{
-		return FMediaTimeStamp(Time);
+		return Time;
 	}
 
 	virtual EMediaOverlaySampleType GetType() const override
@@ -116,5 +116,5 @@ private:
 	FText Text;
 
 	/** Presentation time for which the sample was generated. */
-	FTimespan Time;
+	FMediaTimeStamp Time;
 };
diff --git a/Engine/Plugins/Media/WmfMedia/Source/WmfMedia/Private/Player/WmfMediaTextureSample.h b/Engine/Plugins/Media/WmfMedia/Source/WmfMedia/Private/Player/WmfMediaTextureSample.h
--- a/Engine/Plugins/Media/WmfMedia/Source/WmfMedia/Private/Player/WmfMediaTextureSample.h
+++ b/Engine/Plugins/Media/WmfMedia/Source/WmfMedia/Private/Player/WmfMediaTextureSample.h
@@ -27,7 +27,6 @@ public:
 		, OutputDim(FIntPoint::ZeroValue)
 		, SampleFormat(EMediaTextureSampleFormat::Undefined)
 		, Time(FTimespan::Zero())
-		, SequenceIndex(0)
 		, bIsBufferExternal(false)
 		, Stride(0)
 	{ }
@@ -56,7 +55,7 @@ public:
 		const FIntPoint& InOutputDim,
 		EMediaTextureSampleFormat InSampleFormat,
 		uint32 InStride,
-		FTimespan InTime,
+		FMediaTimeStamp InTime,
 		FTimespan InDuration)
 	{
 		if ((InBuffer == nullptr) || (InSampleFormat == EMediaTextureSampleFormat::Undefined) || (InSize == 0) || (InStride == 0))
@@ -101,7 +100,7 @@ public:
 		const FIntPoint& InOutputDim,
 		EMediaTextureSampleFormat InSampleFormat,
 		uint32 InStride,
-		FTimespan InTime,
+		FMediaTimeStamp InTime,
 		FTimespan InDuration)
 	{
 		if (InBuffer == nullptr)
@@ -135,11 +134,6 @@ public:
 	 */
 	bool IsBufferExternal() const { return bIsBufferExternal; }
 
-	/**
-	 * Sets the sequence index for the time stamp of this sample.
-	 */
-	void SetSequenceIndex(int64 Index) { SequenceIndex = Index; }
-
 public:
 
 	//~ IMediaTextureSample interface
@@ -183,7 +177,7 @@ public:
 
 	virtual FMediaTimeStamp GetTime() const override
 	{
-		return FMediaTimeStamp(Time, SequenceIndex);
+		return Time;
 	}
 
 	virtual bool IsCacheable() const override
@@ -220,10 +214,7 @@ protected:
 	EMediaTextureSampleFormat SampleFormat;
 
 	/** Presentation for which the sample was generated. */
-	FTimespan Time;
-
-	/** Sequence index used to create the FMediaTimeStamp. */
-	int64 SequenceIndex;
+	FMediaTimeStamp Time;
 
 	/** If true, then Buffer is from WmfMediaDecoder and needs to be returned to it. */
 	bool bIsBufferExternal;
diff --git a/Engine/Plugins/Media/WmfMedia/Source/WmfMedia/Private/Player/WmfMediaTracks.cpp b/Engine/Plugins/Media/WmfMedia/Source/WmfMedia/Private/Player/WmfMediaTracks.cpp
--- a/Engine/Plugins/Media/WmfMedia/Source/WmfMedia/Private/Player/WmfMediaTracks.cpp
+++ b/Engine/Plugins/Media/WmfMedia/Source/WmfMedia/Private/Player/WmfMediaTracks.cpp
@@ -369,9 +369,17 @@ void FWmfMediaTracks::Shutdown()
 
 	MediaSourceChanged = false;
 	SelectionChanged = false;
+
 #if WMFMEDIA_PLAYER_VERSION >= 2
-	SeekTimeOptional.Reset();
 	SeekIndex = 0;
+	AudioLoopIndex = 0;
+	VideoLoopIndex = 0;
+	MetaDataLoopIndex = 0;
+	CaptionLoopIndex = 0;
+	LastAudioTime.Reset();
+	LastVideoTime.Reset();
+	LastMetaDataTime.Reset();
+	LastCaptionTime.Reset();
 #endif // WMFMEDIA_PLAYER_VERSION >= 2
 }
 
@@ -386,7 +394,7 @@ void FWmfMediaTracks::SetSessionState(EMediaState InState)
 void FWmfMediaTracks::SeekStarted(const FTimespan& InTime)
 {
 	UE_LOG(LogWmfMedia, VeryVerbose, TEXT("FWmfMediaTracks::SeekStarted %f"), InTime.GetTotalSeconds());
-	SeekTimeOptional = InTime;
+	FScopeLock Lock(&CriticalSection);
 	++SeekIndex;
 }
 
@@ -514,152 +522,9 @@ void FWmfMediaTracks::FlushSamples()
 
 #if WMFMEDIA_PLAYER_VERSION >= 2
 
-IMediaSamples::EFetchBestSampleResult FWmfMediaTracks::FetchBestVideoSampleForTimeRange(const TRange<FMediaTimeStamp> & TimeRange, TSharedPtr<IMediaTextureSample, ESPMode::ThreadSafe>& OutSample, bool bReverse)
+IMediaSamples::EFetchBestSampleResult FWmfMediaTracks::FetchBestVideoSampleForTimeRange(const TRange<FMediaTimeStamp>& TimeRange, TSharedPtr<IMediaTextureSample, ESPMode::ThreadSafe>& OutSample, bool bReverse, bool bConsistentResult)
 {
-	// Don't return any samples if we are stopped. We could be prerolling.
-	if (SessionState == EMediaState::Stopped)
-	{
-		return IMediaSamples::EFetchBestSampleResult::NoSample;;
-	}
-
-	FTimespan TimeRangeLow = TimeRange.GetLowerBoundValue().Time;
-	FTimespan TimeRangeHigh = TimeRange.GetUpperBoundValue().Time;
-	// Account for loop wraparound.
-	if (TimeRangeHigh < TimeRangeLow)
-	{
-		TimeRangeHigh += CachedDuration;
-	}
-	TRange<FTimespan> TimeRangeTime(TimeRangeLow, TimeRangeHigh);
-	FTimespan LoopDiff = CachedDuration * 0.5f;
-	float CurrentOverlap = 0.0f;
-	IMediaSamples::EFetchBestSampleResult Result = IMediaSamples::EFetchBestSampleResult::NoSample;
-	UE_LOG(LogWmfMedia, VeryVerbose, TEXT("FetchBestVideoSampleForTimeRange %f:%d %f:%d seek:%f"),
-		TimeRangeLow.GetTotalSeconds(), TimeRange.GetLowerBoundValue().SequenceIndex, TimeRangeHigh.GetTotalSeconds(), TimeRange.GetUpperBoundValue().SequenceIndex,
-		SeekTimeOptional.IsSet() ? SeekTimeOptional->GetTotalSeconds() : -1.0f);
-
-	// Loop over our samples.
-	while (true)
-	{
-		// Is there a sample available?
-		TSharedPtr<IMediaTextureSample, ESPMode::ThreadSafe> Sample;
-		if (VideoSampleQueue.Peek(Sample))
-		{
-			FTimespan SampleStartTime = Sample->GetTime().Time;
-			FTimespan SampleEndTime = SampleStartTime + Sample->GetDuration();
-			UE_LOG(LogWmfMedia, VeryVerbose, TEXT("FetchBestVideoSampleForTimeRange looking at sample %f:%d %f"),
-				SampleStartTime.GetTotalSeconds(), Sample->GetTime().SequenceIndex, SampleEndTime.GetTotalSeconds());
-
-#if WMFMEDIA_PLAYER_VERSION >= 2
-			// Are we waiting for the sample from a seek?
-			if (SeekTimeOptional.IsSet())
-			{
-				// Are we past the seek time?
-				if (TimeRangeTime.Contains(SeekTimeOptional.GetValue()))
-				{
-					// Is this our seek sample?
-					FTimespan SeekTime = SeekTimeOptional.GetValue();
-					double SeekTimeSeconds = SeekTime.GetTotalSeconds();
-					if ((FMath::IsNearlyEqual(SeekTimeSeconds, SampleStartTime.GetTotalSeconds(), 0.001)) ||
-						((SeekTime >= SampleStartTime) && (SeekTime < SampleEndTime)))
-					{
-						// Yes this is what we have been waiting for.
-						// Reset the seek time so its no longer used.
-						SeekTimeOptional.Reset();
-					}
-					else
-					{
-						// This is not the sample we want, its old.
-						VideoSampleQueue.Pop();
-						continue;
-					}
-				}
-			}
-#endif // WMFMEDIA_PLAYER_VERSION >= 2
-
-			// Are we already past this sample?
-			if (SampleEndTime < TimeRangeLow)
-			{
-				// If there is a large gap to this sample, then its probably because it looped,
-				// so we aren't really past it.
-				FTimespan Diff = TimeRangeLow - SampleEndTime;
-				if (Diff > LoopDiff)
-				{
-					// Adjust sample times so they are in the same "space" as the time range.
-					SampleStartTime += CachedDuration;
-					SampleEndTime += CachedDuration;
-					UE_LOG(LogWmfMedia, VeryVerbose, TEXT("FetchBestVideoSampleForTimeRange sample loop %f %f"),
-						SampleStartTime.GetTotalSeconds(), SampleEndTime.GetTotalSeconds());
-				}
-				else
-				{
-					// Try next sample.
-					VideoSampleQueue.Pop();
-					continue;
-				}
-			}
-			
-			{
-#if WMFMEDIA_PLAYER_VERSION >= 2
-				// Did we already pass this sample,
-				// and the sample is at the end of the video and we just looped around?
-				FTimespan Diff = SampleEndTime - TimeRangeLow;
-				if (Diff > LoopDiff)
-				{
-					VideoSampleQueue.Pop();
-					continue;
-				}
-
-				
-#endif // WMFMEDIA_PLAYER_VERSION >= 2
-
-				// Is this sample before the end of the requested time range?
-				if (SampleStartTime < TimeRangeHigh)
-				{
-					// Yes.
-					// Does this sample have the largest overlap so far?
-					TRange<FTimespan> SampleRange(SampleStartTime, SampleEndTime);
-					TRange<FTimespan> OverlapRange = TRange<FTimespan>::Intersection(SampleRange, TimeRangeTime);
-					FTimespan OverlapTimespan = OverlapRange.Size<FTimespan>();
-					float Overlap = OverlapTimespan.GetTotalSeconds();
-					if (CurrentOverlap <= Overlap)
-					{
-						// Yes. Use this sample.
-						if (VideoSampleQueue.Dequeue(OutSample))
-						{
-							Result = IMediaSamples::EFetchBestSampleResult::Ok;
-							CurrentOverlap = Overlap;
-							
-							// Update sequence index.
-							FWmfMediaTextureSample* WmfSample =
-								static_cast<FWmfMediaTextureSample*>(OutSample.Get());
-							WmfSample->SetSequenceIndex(FMediaTimeStamp::MakeSequenceIndex(SeekIndex, 0));
-
-							UE_LOG(LogWmfMedia, VeryVerbose, TEXT("FetchBestVideoSampleForTimeRange got sample."));
-						}
-					}
-					else
-					{
-						// No need to continue.
-						// This sample is overlapping our end point.
-						break;
-					}
-				}
-				else
-				{
-					// Sample is not before the end of the requested time range.
-					// We are done for now.
-					break;
-				}
-			}
-		}
-		else
-		{
-			// No samples available.
-			break;
-		}
-	}
-
-	return Result;
+	return VideoSampleQueue.FetchBestSampleForTimeRange(TimeRange, OutSample, bReverse, bConsistentResult) ? EFetchBestSampleResult::Ok : EFetchBestSampleResult::NoSample;
 }
 
 #endif // WMFMEDIA_PLAYER_VERSION >= 2
@@ -671,10 +536,30 @@ bool FWmfMediaTracks::PeekVideoSampleTime(FMediaTimeStamp & TimeStamp)
 	{
 		return false;
 	}
-	TimeStamp = FMediaTimeStamp(Sample->GetTime());
+	TimeStamp = Sample->GetTime();
 	return true;
 }
 
+bool FWmfMediaTracks::DiscardVideoSamples(const TRange<FMediaTimeStamp>& TimeRange, bool bReverse)
+{
+	return VideoSampleQueue.Discard(TimeRange, bReverse);
+}
+
+bool FWmfMediaTracks::DiscardAudioSamples(const TRange<FMediaTimeStamp>& TimeRange, bool bReverse)
+{
+	return AudioSampleQueue.Discard(TimeRange, bReverse);
+}
+
+bool FWmfMediaTracks::DiscardCaptionSamples(const TRange<FMediaTimeStamp>& TimeRange, bool bReverse)
+{
+	return CaptionSampleQueue.Discard(TimeRange, bReverse);
+}
+
+bool FWmfMediaTracks::DiscardMetadataSamples(const TRange<FMediaTimeStamp>& TimeRange, bool bReverse)
+{
+	return MetadataSampleQueue.Discard(TimeRange, bReverse);
+}
+
 /* IMediaTracks interface
  *****************************************************************************/
 
@@ -1288,7 +1173,10 @@ bool FWmfMediaTracks::AddTrackToTopology(const FTrack& Track, IMFTopology& Topol
 	{
 		VideoHardwareVideoDecodingSamplePool = MakeShared<FWmfMediaHardwareVideoDecodingTextureSamplePool>();
 
-		MediaStreamSink->SetMediaSamplePoolAndQueue(VideoHardwareVideoDecodingSamplePool, &VideoSampleQueue);
+		MediaStreamSink->SetMediaSamplePoolAndQueue(VideoHardwareVideoDecodingSamplePool, &VideoSampleQueue, [this](FTimespan Time, EMediaTrackType TrackType) -> FMediaTimeStamp
+		{
+			return AdjustTimeStamp(Time, TrackType);
+		});
 
 		if (FAILED(::MFCreateTopologyNode(MF_TOPOLOGY_OUTPUT_NODE, &OutputNode)) ||
 			FAILED(OutputNode->SetObject((IMFStreamSink*)MediaStreamSink)) ||
@@ -2007,6 +1895,34 @@ void FWmfMediaTracks::HandleMediaSamplerClock(EWmfMediaSamplerClockEvent Event,
 }
 
 
+FMediaTimeStamp FWmfMediaTracks::AdjustTimeStamp(FTimespan Time, EMediaTrackType TrackType)
+{
+	FScopeLock Lock(&CriticalSection);
+
+	int32* IndexPtr = nullptr;
+	TOptional<FTimespan>* TimePtr = nullptr;
+	switch (TrackType)
+	{
+	case EMediaTrackType::Audio: IndexPtr = &AudioLoopIndex; TimePtr = &LastAudioTime; break;
+	case EMediaTrackType::Video: IndexPtr = &VideoLoopIndex; TimePtr = &LastVideoTime; break;
+	case EMediaTrackType::Metadata: IndexPtr = &MetaDataLoopIndex; TimePtr = &LastMetaDataTime; break;
+	case EMediaTrackType::Caption: IndexPtr = &CaptionLoopIndex; TimePtr = &LastCaptionTime; break;
+	}
+	check(IndexPtr && TimePtr);
+
+	if (TimePtr->IsSet())
+	{
+		if (TimePtr->GetValue() > Time)
+		{
+			++(*IndexPtr);
+		}
+	}
+	*TimePtr = Time;
+
+	return FMediaTimeStamp(Time, FMediaTimeStamp::MakeSequenceIndex(SeekIndex, *IndexPtr));
+}
+
+
 void FWmfMediaTracks::HandleMediaSamplerAudioSample(const uint8* Buffer, uint32 Size, FTimespan /*Duration*/, FTimespan Time)
 {
 	if (Buffer == nullptr)
@@ -2040,7 +1956,7 @@ void FWmfMediaTracks::HandleMediaSamplerAudioSample(const uint8* Buffer, uint32
 	// create & add sample to queue
 	const TSharedRef<FWmfMediaAudioSample, ESPMode::ThreadSafe> AudioSample = AudioSamplePool->AcquireShared();
 
-	if (AudioSample->Initialize(Buffer, Size, Format->Audio.NumChannels, Format->Audio.SampleRate, Time, Duration))
+	if (AudioSample->Initialize(Buffer, Size, Format->Audio.NumChannels, Format->Audio.SampleRate, AdjustTimeStamp(Time, EMediaTrackType::Audio), Duration))
 	{
 		AudioSampleQueue.Enqueue(AudioSample);
 	}
@@ -2070,7 +1986,7 @@ void FWmfMediaTracks::HandleMediaSamplerCaptionSample(const uint8* Buffer, uint3
 	const FTrack& Track = CaptionTracks[SelectedCaptionTrack];
 	const auto CaptionSample = MakeShared<FWmfMediaOverlaySample, ESPMode::ThreadSafe>();
 
-	if (CaptionSample->Initialize((char*)Buffer, Time, Duration))
+	if (CaptionSample->Initialize((char*)Buffer, AdjustTimeStamp(Time, EMediaTrackType::Caption), Duration))
 	{
 		CaptionSampleQueue.Enqueue(CaptionSample);
 	}
@@ -2100,7 +2016,7 @@ void FWmfMediaTracks::HandleMediaSamplerMetadataSample(const uint8* Buffer, uint
 	const FTrack& Track = MetadataTracks[SelectedMetadataTrack];
 	const auto BinarySample = MakeShared<FWmfMediaBinarySample, ESPMode::ThreadSafe>();
 
-	if (BinarySample->Initialize(Buffer, Size, Time, Duration))
+	if (BinarySample->Initialize(Buffer, Size, AdjustTimeStamp(Time, EMediaTrackType::Metadata), Duration))
 	{
 		MetadataSampleQueue.Enqueue(BinarySample);
 	}
@@ -2162,7 +2078,7 @@ void FWmfMediaTracks::HandleMediaSamplerVideoSample(const uint8* Buffer, uint32
 		Format->Video.OutputDim,
 		Format->Video.SampleFormat,
 		Format->Video.BufferStride,
-		Time,
+		AdjustTimeStamp(Time, EMediaTrackType::Video),
 		Duration))
 	{
 		VideoSampleQueue.Enqueue(TextureSample);
diff --git a/Engine/Plugins/Media/WmfMedia/Source/WmfMedia/Private/Player/WmfMediaTracks.h b/Engine/Plugins/Media/WmfMedia/Source/WmfMedia/Private/Player/WmfMediaTracks.h
--- a/Engine/Plugins/Media/WmfMedia/Source/WmfMedia/Private/Player/WmfMediaTracks.h
+++ b/Engine/Plugins/Media/WmfMedia/Source/WmfMedia/Private/Player/WmfMediaTracks.h
@@ -213,8 +213,12 @@ public:
 
 	virtual void FlushSamples() override;
 #if WMFMEDIA_PLAYER_VERSION >= 2
-	virtual EFetchBestSampleResult FetchBestVideoSampleForTimeRange(const TRange<FMediaTimeStamp> & TimeRange, TSharedPtr<IMediaTextureSample, ESPMode::ThreadSafe>& OutSample, bool bReverse) override;
+	virtual EFetchBestSampleResult FetchBestVideoSampleForTimeRange(const TRange<FMediaTimeStamp>& TimeRange, TSharedPtr<IMediaTextureSample, ESPMode::ThreadSafe>& OutSample, bool bReverse, bool bConsistentResult) override;
 #endif // WMFMEDIA_PLAYER_VERSION >= 2
+	virtual bool DiscardVideoSamples(const TRange<FMediaTimeStamp>& TimeRange, bool bReverse) override;
+	virtual bool DiscardAudioSamples(const TRange<FMediaTimeStamp>& TimeRange, bool bReverse) override;
+	virtual bool DiscardCaptionSamples(const TRange<FMediaTimeStamp>& TimeRange, bool bReverse) override;
+	virtual bool DiscardMetadataSamples(const TRange<FMediaTimeStamp>& TimeRange, bool bReverse) override;
 	virtual bool PeekVideoSampleTime(FMediaTimeStamp & TimeStamp) override;
 
 public:
@@ -305,6 +309,9 @@ private:
 	/** Callback for handling new video samples. */
 	void HandleMediaSamplerVideoSample(const uint8* Buffer, uint32 Size, FTimespan Duration, FTimespan Time);
 
+	/** Make sure sequence index is up to date and combine with given timestamp */
+	FMediaTimeStamp AdjustTimeStamp(FTimespan Time, EMediaTrackType TrackType);
+
 private:
 
 	/** Audio sample object pool. */
@@ -381,10 +388,16 @@ private:
 	EMediaState SessionState;
 
 #if WMFMEDIA_PLAYER_VERSION >= 2
-	/** If set, then discard samples until we get this sample. */
-	TOptional<FTimespan> SeekTimeOptional;
-	/** Seek index tracking */
+	/** SequenceIndex tracking */
 	int32 SeekIndex;
+	int32 AudioLoopIndex;
+	int32 VideoLoopIndex;
+	int32 MetaDataLoopIndex;
+	int32 CaptionLoopIndex;
+	TOptional<FTimespan> LastAudioTime;
+	TOptional<FTimespan> LastVideoTime;
+	TOptional<FTimespan> LastMetaDataTime;
+	TOptional<FTimespan> LastCaptionTime;
 #endif // WMFMEDIA_PLAYER_VERSION >= 2
 };
 
diff --git a/Engine/Plugins/Media/WmfMedia/Source/WmfMedia/Private/Wmf/WmfMediaHardwareVideoDecodingTextureSample.cpp b/Engine/Plugins/Media/WmfMedia/Source/WmfMedia/Private/Wmf/WmfMediaHardwareVideoDecodingTextureSample.cpp
--- a/Engine/Plugins/Media/WmfMedia/Source/WmfMedia/Private/Wmf/WmfMediaHardwareVideoDecodingTextureSample.cpp
+++ b/Engine/Plugins/Media/WmfMedia/Source/WmfMedia/Private/Wmf/WmfMediaHardwareVideoDecodingTextureSample.cpp
@@ -6,7 +6,7 @@
 
 #if WMFMEDIA_SUPPORTED_PLATFORM
 
-ID3D11Texture2D* FWmfMediaHardwareVideoDecodingTextureSample::InitializeSourceTexture(const TRefCountPtr<ID3D11Device>& InD3D11Device, FTimespan InTime, FTimespan InDuration, const FIntPoint& InDim, EPixelFormat InFormat, EMediaTextureSampleFormat InMediaTextureSampleFormat)
+ID3D11Texture2D* FWmfMediaHardwareVideoDecodingTextureSample::InitializeSourceTexture(const TRefCountPtr<ID3D11Device>& InD3D11Device, FMediaTimeStamp InTime, FTimespan InDuration, const FIntPoint& InDim, EPixelFormat InFormat, EMediaTextureSampleFormat InMediaTextureSampleFormat)
 {
 	Time = InTime;
 	Dim = InDim;
diff --git a/Engine/Plugins/Media/WmfMedia/Source/WmfMedia/Private/Wmf/WmfMediaSession.cpp b/Engine/Plugins/Media/WmfMedia/Source/WmfMedia/Private/Wmf/WmfMediaSession.cpp
--- a/Engine/Plugins/Media/WmfMedia/Source/WmfMedia/Private/Wmf/WmfMediaSession.cpp
+++ b/Engine/Plugins/Media/WmfMedia/Source/WmfMedia/Private/Wmf/WmfMediaSession.cpp
@@ -53,6 +53,7 @@ FWmfMediaSession::FWmfMediaSession()
 	, PendingChanges(false)
 	, RefCount(0)
 	, bIsRequestedTimeLoop(false)
+	, bIsRequestedTimeSeek(false)
 	, SessionRate(0.0f)
 	, SessionState(EMediaState::Closed)
 	, ShouldLoop(false)
@@ -497,6 +498,7 @@ bool FWmfMediaSession::Seek(const FTimespan& Time)
 		UE_LOG(LogWmfMedia, Verbose, TEXT("Session %p: Requesting seek after pending command"), this);
 		RequestedTime = Time;
 		bIsRequestedTimeLoop = false;
+		bIsRequestedTimeSeek = true;
 		return true;
 	}
 
@@ -840,6 +842,7 @@ bool FWmfMediaSession::CommitRate(float Rate)
 			{			
 				RequestedTime = LastTime;
 				bIsRequestedTimeLoop = false;
+				bIsRequestedTimeSeek = false;
 			}
 
 			RequestedRate = Rate;
@@ -957,6 +960,7 @@ bool FWmfMediaSession::CommitRate(float Rate)
 			UE_LOG(LogWmfMedia, Verbose, TEXT("Session %p: Requesting start after pending rate change"), this);
 			RequestedTime = RestartTime;
 			bIsRequestedTimeLoop = false;
+			bIsRequestedTimeSeek = false;
 		}
 		else
 		{
@@ -976,7 +980,7 @@ bool FWmfMediaSession::CommitTime(FTimespan Time, bool bIsSeek)
 
 	FTimespan OriginalTime = Time;
 	const FString TimeString = (Time == WmfMediaSession::RequestedTimeCurrent) ? TEXT("<current>") : *Time.ToString();
-	UE_LOG(LogWmfMedia, Verbose, TEXT("Session %p: Committing time %s"), this, *TimeString);
+	UE_LOG(LogWmfMedia, Verbose, TEXT("Session %p: Committing time %s (Seek=%d)"), this, *TimeString, bIsSeek);
 
 	// start session at requested time
 	PROPVARIANT StartPosition;
@@ -1020,7 +1024,7 @@ bool FWmfMediaSession::CommitTime(FTimespan Time, bool bIsSeek)
 
 #if WMFMEDIA_PLAYER_VERSION >= 2
 	// If this is not a loop, then tell the tracks about the seek.
-	if ((bIsRequestedTimeLoop == false) && (bCanSeek) && bIsSeek)
+	if (bCanSeek && bIsSeek)
 	{
 		TSharedPtr<FWmfMediaTracks, ESPMode::ThreadSafe> TracksPinned = Tracks.Pin();
 		if (TracksPinned.IsValid())
@@ -1065,6 +1069,7 @@ bool FWmfMediaSession::CommitTopology(IMFTopology* Topology)
 
 			RequestedTime = LastTime;
 			bIsRequestedTimeLoop = false;
+			bIsRequestedTimeSeek = false;
 
 			// Zero LastTime so it matches what WMF is actually doing (rewinding to the beginning)
 			// The PresentationClock has now stopped and FWmfMediaSession::GetTime() will just return LastTime
@@ -1114,6 +1119,7 @@ void FWmfMediaSession::DiscardPendingChanges()
 	RequestedRate.Reset();
 	RequestedTime.Reset();
 	bIsRequestedTimeLoop = false;
+	bIsRequestedTimeSeek = false;
 	RequestedTopology.Reset();
 
 	PendingChanges = false;
@@ -1171,8 +1177,9 @@ void FWmfMediaSession::DoPendingChanges()
 		const FTimespan Time = RequestedTime.GetValue();
 		RequestedTime.Reset();
 
-		CommitTime(Time, false);
+		CommitTime(Time, bIsRequestedTimeSeek && !bIsRequestedTimeLoop);
 		bIsRequestedTimeLoop = false;
+		bIsRequestedTimeSeek = false;
 	}
 }
 
@@ -1318,6 +1325,7 @@ void FWmfMediaSession::HandleSessionEnded()
 		// loop back to beginning/end
 		RequestedTime = (SessionRate < 0.0f) ? CurrentDuration : FTimespan::Zero();
 		bIsRequestedTimeLoop = true;
+		bIsRequestedTimeSeek = false;
 		DoPendingChanges();
 	}
 	else
@@ -1594,6 +1602,7 @@ void FWmfMediaSession::HandleSessionTopologyStatus(HRESULT EventStatus, IMFMedia
 
 			RequestedTime = WmfMediaSession::RequestedTimeCurrent;
 			bIsRequestedTimeLoop = false;
+			bIsRequestedTimeSeek = false;
 			PendingChanges = true;
 		}
 	}
diff --git a/Engine/Plugins/Media/WmfMedia/Source/WmfMedia/Private/Wmf/WmfMediaSession.h b/Engine/Plugins/Media/WmfMedia/Source/WmfMedia/Private/Wmf/WmfMediaSession.h
--- a/Engine/Plugins/Media/WmfMedia/Source/WmfMedia/Private/Wmf/WmfMediaSession.h
+++ b/Engine/Plugins/Media/WmfMedia/Source/WmfMedia/Private/Wmf/WmfMediaSession.h
@@ -260,6 +260,9 @@ private:
 	/** If true then RequestedTime is due to the video looping. */
 	bool bIsRequestedTimeLoop;
 
+	/** If true then RequestedTime is due to the video seeking. */
+	bool bIsRequestedTimeSeek;
+
 	/** The session's internal playback rate (not necessarily the same as GetRate). */
 	float SessionRate;
 
diff --git a/Engine/Plugins/Media/WmfMedia/Source/WmfMedia/Private/Wmf/WmfMediaStreamSink.cpp b/Engine/Plugins/Media/WmfMedia/Source/WmfMedia/Private/Wmf/WmfMediaStreamSink.cpp
--- a/Engine/Plugins/Media/WmfMedia/Source/WmfMedia/Private/Wmf/WmfMediaStreamSink.cpp
+++ b/Engine/Plugins/Media/WmfMedia/Source/WmfMedia/Private/Wmf/WmfMediaStreamSink.cpp
@@ -966,7 +966,7 @@ void FWmfMediaStreamSink::CopyTextureAndEnqueueSample(IMFSample* pSample)
 			TextureSample->InitializeExternal(&ExternalBuffer,
 				FIntPoint(DimX, DimY), FIntPoint(DimX, DimY), MediaTextureSampleFormat,
 				Pitch,
-				FTimespan::FromMicroseconds(SampleTime / 10),
+				AdjustTimeStamp(FTimespan::FromMicroseconds(SampleTime / 10), EMediaTrackType::Video),
 				FTimespan::FromMicroseconds(SampleDuration / 10));
 
 			TextureSample->SetPixelFormat(PixelFormat);
@@ -1000,7 +1000,7 @@ void FWmfMediaStreamSink::CopyTextureAndEnqueueSample(IMFSample* pSample)
 
 				ID3D11Texture2D* SharedTexture = TextureSample->InitializeSourceTexture(
 					Owner->GetDevice(),
-					FTimespan::FromMicroseconds(SampleTime / 10),
+					AdjustTimeStamp(FTimespan::FromMicroseconds(SampleTime / 10), EMediaTrackType::Video),
 					FTimespan::FromMicroseconds(SampleDuration / 10),
 					FIntPoint(DimX, DimY),
 					PixelFormat,
@@ -1064,7 +1064,7 @@ void FWmfMediaStreamSink::CopyTextureAndEnqueueSample(IMFSample* pSample)
 					}
 
 					uint32 Pitch = (DimX / GPixelFormats[PixelFormat].BlockSizeX) * GPixelFormats[PixelFormat].BlockBytes;
-					TextureSample->Initialize(Data, BufferSize, FIntPoint(DimX, SampleDimY), FIntPoint(DimX, DimY), MediaTextureSampleFormat, Pitch, FTimespan::FromMicroseconds(SampleTime / 10), FTimespan::FromMicroseconds(SampleDuration / 10));
+					TextureSample->Initialize(Data, BufferSize, FIntPoint(DimX, SampleDimY), FIntPoint(DimX, DimY), MediaTextureSampleFormat, Pitch, AdjustTimeStamp(FTimespan::FromMicroseconds(SampleTime / 10), EMediaTrackType::Video), FTimespan::FromMicroseconds(SampleDuration / 10));
 					pBuffer->Unlock();
 
 					TextureSample->SetPixelFormat(PixelFormat);
@@ -1129,11 +1129,13 @@ void FWmfMediaStreamSink::SetClockRate(float InClockRate)
 
 void FWmfMediaStreamSink::SetMediaSamplePoolAndQueue(
 	TSharedPtr<FWmfMediaHardwareVideoDecodingTextureSamplePool>& InVideoSamplePool,
-	TMediaSampleQueue<IMediaTextureSample>* InVideoSampleQueue)
+	TMediaSampleQueue<IMediaTextureSample>* InVideoSampleQueue,
+	TFunction<FMediaTimeStamp(FTimespan, EMediaTrackType)>&& InAdjustTimeStamp)
 {
 	FScopeLock Lock(&CriticalSection);
 	VideoSamplePool = InVideoSamplePool;
 	VideoSampleQueue = InVideoSampleQueue;
+	AdjustTimeStamp = MoveTemp(InAdjustTimeStamp);
 }
 
 
diff --git a/Engine/Plugins/Media/WmfMedia/Source/WmfMedia/Private/Wmf/WmfMediaStreamSink.h b/Engine/Plugins/Media/WmfMedia/Source/WmfMedia/Private/Wmf/WmfMediaStreamSink.h
--- a/Engine/Plugins/Media/WmfMedia/Source/WmfMedia/Private/Wmf/WmfMediaStreamSink.h
+++ b/Engine/Plugins/Media/WmfMedia/Source/WmfMedia/Private/Wmf/WmfMediaStreamSink.h
@@ -206,8 +206,9 @@ public:
 	 *
 	 * @param InVideoSamplePool sample pool to get sample from
 	 * @param InVideoSampleQueue Sample queue to fill
+	 * @param InSequenceIndex Pointer to current sequence index value
 	 */
-	void SetMediaSamplePoolAndQueue(TSharedPtr<FWmfMediaHardwareVideoDecodingTextureSamplePool>& InVideoSamplePool, TMediaSampleQueue<IMediaTextureSample>* InVideoSampleQueue);
+	void SetMediaSamplePoolAndQueue(TSharedPtr<FWmfMediaHardwareVideoDecodingTextureSamplePool>& InVideoSamplePool, TMediaSampleQueue<IMediaTextureSample>* InVideoSampleQueue, TFunction<FMediaTimeStamp(FTimespan, EMediaTrackType)>&& InAdjustTimeStamp);
 
 private:
 
@@ -281,8 +282,11 @@ private:
 	/** Video sample pool from which to get next free sample */
 	TSharedPtr<FWmfMediaHardwareVideoDecodingTextureSamplePool> VideoSamplePool;
 
-	/** Video sample queue which is filled with current sample*/
+	/** Video sample queue which is filled with current sample */
 	TMediaSampleQueue<IMediaTextureSample>* VideoSampleQueue;
+	
+	/** Callback to adjust given timestamp for the output queue */
+	TFunction<FMediaTimeStamp(FTimespan,EMediaTrackType)> AdjustTimeStamp;
 
 	/** Make sure we don't spam output with MF_MT_SUBTYPE error message */
 	bool bShowSubTypeErrorMessage;
diff --git a/Engine/Plugins/Media/WmfMedia/Source/WmfMedia/Public/WmfMediaHardwareVideoDecodingTextureSample.h b/Engine/Plugins/Media/WmfMedia/Source/WmfMedia/Public/WmfMediaHardwareVideoDecodingTextureSample.h
--- a/Engine/Plugins/Media/WmfMedia/Source/WmfMedia/Public/WmfMediaHardwareVideoDecodingTextureSample.h
+++ b/Engine/Plugins/Media/WmfMedia/Source/WmfMedia/Public/WmfMediaHardwareVideoDecodingTextureSample.h
@@ -50,7 +50,7 @@ public:
 	 * @param InCreateFlags texture create flag
 	 * @return The texture resource object that will hold the sample data.
 	 */
-	ID3D11Texture2D* InitializeSourceTexture(const TRefCountPtr<ID3D11Device>& InD3D11Device, FTimespan InTime, FTimespan InDuration, const FIntPoint& InDim, EPixelFormat InFormat, EMediaTextureSampleFormat InMediaTextureSampleFormat);
+	ID3D11Texture2D* InitializeSourceTexture(const TRefCountPtr<ID3D11Device>& InD3D11Device, FMediaTimeStamp InTime, FTimespan InDuration, const FIntPoint& InDim, EPixelFormat InFormat, EMediaTextureSampleFormat InMediaTextureSampleFormat);
 
 	/**
 	 * Call this to set the pixel format.
diff --git a/Engine/Source/Runtime/Media/Public/IMediaSamples.h b/Engine/Source/Runtime/Media/Public/IMediaSamples.h
--- a/Engine/Source/Runtime/Media/Public/IMediaSamples.h
+++ b/Engine/Source/Runtime/Media/Public/IMediaSamples.h
@@ -122,7 +122,7 @@ public:
 		NoSample,
 		NotSupported,
 	};
-	virtual EFetchBestSampleResult FetchBestVideoSampleForTimeRange(const TRange<FMediaTimeStamp> & TimeRange, TSharedPtr<IMediaTextureSample, ESPMode::ThreadSafe>& OutSample, bool bReverse)
+	virtual EFetchBestSampleResult FetchBestVideoSampleForTimeRange(const TRange<FMediaTimeStamp>& TimeRange, TSharedPtr<IMediaTextureSample, ESPMode::ThreadSafe>& OutSample, bool bReverse, bool bConsistentResult)
 	{
 		return EFetchBestSampleResult::NotSupported;
 	}
diff --git a/Engine/Source/Runtime/Media/Public/IMediaTimeSource.h b/Engine/Source/Runtime/Media/Public/IMediaTimeSource.h
--- a/Engine/Source/Runtime/Media/Public/IMediaTimeSource.h
+++ b/Engine/Source/Runtime/Media/Public/IMediaTimeSource.h
@@ -28,8 +28,8 @@ class FMediaTimeStamp
 {
 public:
 	FMediaTimeStamp() : Time(FTimespan::MinValue()), SequenceIndex(0) {}
-	FMediaTimeStamp(const FTimespan & InTime) : Time(InTime), SequenceIndex(0) {}
-	FMediaTimeStamp(const FTimespan & InTime, int64 InSequenceIndex) : Time(InTime), SequenceIndex(InSequenceIndex) {}
+	explicit FMediaTimeStamp(const FTimespan & InTime) : Time(InTime), SequenceIndex(0) {}
+	explicit FMediaTimeStamp(const FTimespan & InTime, int64 InSequenceIndex) : Time(InTime), SequenceIndex(InSequenceIndex) {}
 
 	void Invalidate() { Time = FTimespan::MinValue(); }
 	bool IsValid() const { return Time != FTimespan::MinValue(); }
diff --git a/Engine/Source/Runtime/MediaAssets/Private/Assets/MediaTexture.cpp b/Engine/Source/Runtime/MediaAssets/Private/Assets/MediaTexture.cpp
--- a/Engine/Source/Runtime/MediaAssets/Private/Assets/MediaTexture.cpp
+++ b/Engine/Source/Runtime/MediaAssets/Private/Assets/MediaTexture.cpp
@@ -439,7 +439,7 @@ void UMediaTexture::TickResource(FTimespan Timecode)
 					RenderParams.SampleSource = SampleQueue;
 
 					RenderParams.Rate = CurrentPlayerPtr->GetRate();
-					RenderParams.Time = CurrentPlayerPtr->GetTime();
+					RenderParams.Time = FMediaTimeStamp(CurrentPlayerPtr->GetTime());
 				}
 				else
 				{
diff --git a/Engine/Source/Runtime/MediaAssets/Private/Misc/MediaTextureResource.cpp b/Engine/Source/Runtime/MediaAssets/Private/Misc/MediaTextureResource.cpp
--- a/Engine/Source/Runtime/MediaAssets/Private/Misc/MediaTextureResource.cpp
+++ b/Engine/Source/Runtime/MediaAssets/Private/Misc/MediaTextureResource.cpp
@@ -403,8 +403,8 @@ void FMediaTextureResource::Render(const FRenderParams& Params)
 			TSharedPtr<IMediaTextureSample, ESPMode::ThreadSafe> TestSample;
 			while (SampleSource->Peek(TestSample) && TestSample.IsValid())
 			{
-				const FTimespan StartTime = TestSample->GetTime().Time;
-				const FTimespan EndTime = StartTime + TestSample->GetDuration();
+				const FMediaTimeStamp& StartTime = TestSample->GetTime();
+				const FMediaTimeStamp EndTime = StartTime + TestSample->GetDuration();
 
 				if ((Params.Rate >= 0.0f) && (Params.Time < StartTime))
 				{
diff --git a/Engine/Source/Runtime/MediaUtils/Private/MediaPlayerFacade.cpp b/Engine/Source/Runtime/MediaUtils/Private/MediaPlayerFacade.cpp
--- a/Engine/Source/Runtime/MediaUtils/Private/MediaPlayerFacade.cpp
+++ b/Engine/Source/Runtime/MediaUtils/Private/MediaPlayerFacade.cpp
@@ -1251,16 +1251,17 @@ void FMediaPlayerFacade::FBlockOnRange::OnFlush()
 {
 	LastTimeRange = TRange<FTimespan>::Empty();
 	OnBlockPrimaryIndex = 0;
+	OnBlockSecondaryIndex = 0;
+	RangeIsDirty = true;
 }
 
 
 void FMediaPlayerFacade::FBlockOnRange::OnSeek(int32 PrimaryIndex)
 {
-	TSharedPtr<IMediaPlayer, ESPMode::ThreadSafe> CurrentPlayer(Facade->Player);
-	check(CurrentPlayer.IsValid());
-
 	LastTimeRange = TRange<FTimespan>::Empty();
 	OnBlockPrimaryIndex = PrimaryIndex;
+	OnBlockSecondaryIndex = 0;
+	RangeIsDirty = true;
 }
 
 
@@ -1325,7 +1326,7 @@ const TRange<FMediaTimeStamp>& FMediaPlayerFacade::FBlockOnRange::GetRange() con
 	* On the synthesized sequence and loop index values:
 	* - We track seeks and hence can insert the proper seek index easily, although the user does not provide it / does not need to track it
 	* - The loop index gets somewhat of a special treatment:
-	*  -- With tools like Sequencer a blocked range my speed along the timeline quite quickly and if the player is configured as looping it would be expected that the video loops while this is done
+	*  -- With tools like Sequencer a blocked range my speed along the time line quite quickly and if the player is configured as looping it would be expected that the video loops while this is done
 	*  -- We hence could loop multiple times within a single update interval
 	*  -- Still we treat any loop (aka: a jump "backwards" without an explicit seek) as a single loop iteration
 	*     (this is easier for any player to work with and provides the same visual results)
@@ -1333,7 +1334,7 @@ const TRange<FMediaTimeStamp>& FMediaPlayerFacade::FBlockOnRange::GetRange() con
 
 	if (!CurrentPlayer->GetControls().IsLooping())
 	{
-		int64 SequenceIndex = FMediaTimeStamp::MakeSequenceIndex(OnBlockPrimaryIndex, 0);
+		int64 SequenceIndex = FMediaTimeStamp::MakeSequenceIndex(OnBlockPrimaryIndex, OnBlockSecondaryIndex);
 		BlockOnRange = TRange<FMediaTimeStamp>(FMediaTimeStamp(Start, SequenceIndex), FMediaTimeStamp(End, SequenceIndex));
 	}
 	else
@@ -1354,40 +1355,62 @@ const TRange<FMediaTimeStamp>& FMediaPlayerFacade::FBlockOnRange::GetRange() con
 		}
 
 
-		float Rate = Facade->GetUnpausedRate();
+		bool bReverse = (Facade->GetUnpausedRate() < 0.0f);
 
 		// Modulo on the time to get it into media's range
 		// (assumes zero-start-time)
 		Start = WrappedModulo(Start, Duration);
 		End = WrappedModulo(End, Duration);
 
-		int64 LoopIdxS = FMath::FloorToInt(CurrentTimeRange.GetLowerBoundValue().GetTotalSeconds() / Duration.GetTotalSeconds());
-		int64 LoopIdxE = FMath::FloorToInt(CurrentTimeRange.GetUpperBoundValue().GetTotalSeconds() / Duration.GetTotalSeconds());
+		int32 LoopIdxS = (int32)FMath::FloorToInt(CurrentTimeRange.GetLowerBoundValue().GetTotalSeconds() / Duration.GetTotalSeconds());
+		int32 LoopIdxE = (int32)FMath::FloorToInt(CurrentTimeRange.GetUpperBoundValue().GetTotalSeconds() / Duration.GetTotalSeconds());
 
-		int32 EndLoopIndex = LoopIdxE - LoopIdxS;
-
-		// Assemble final blocking range
-		auto SeqIndexStart = FMediaTimeStamp::MakeSequenceIndex(OnBlockPrimaryIndex, 0);
-		auto SeqIndexEnd = FMediaTimeStamp::MakeSequenceIndex(OnBlockPrimaryIndex, EndLoopIndex);
-		BlockOnRange = TRange<FMediaTimeStamp>(FMediaTimeStamp(Start, SeqIndexStart), FMediaTimeStamp(End, SeqIndexEnd));
-		check(!BlockOnRange.IsEmpty());
-	}
-
-	// Does the new range overlap with the last one we set?
-	if (!SetBlockOnRange.IsEmpty() && SetBlockOnRange.Overlaps(BlockOnRange))
-	{
-		// Yes, make sure the new range is setup so that it is a "correct" progression given the current playback direction...
-		// (this may go so far as to undo any updates if the "new" one does not adds any range in the playback direction)
-		if (CurrentPlayer->GetControls().GetRate() >= 0.0f)
+		if (!LastTimeRange.IsEmpty())
 		{
-			BlockOnRange = TRange<FMediaTimeStamp>(std::max(BlockOnRange.GetLowerBoundValue(), SetBlockOnRange.GetLowerBoundValue()), std::max(BlockOnRange.GetUpperBoundValue(), SetBlockOnRange.GetUpperBoundValue()));
+			// Adjust loop index base such that, given a playback direction, we can guarantee that the new indices returned fit a monotone progression
+			if (!bReverse)
+			{
+				if (WrappedModulo(LastTimeRange.GetLowerBoundValue(), Duration) > Start)
+				{
+					int32 LastLoopIdxE = (int32)FMath::FloorToInt(LastTimeRange.GetUpperBoundValue().GetTotalSeconds() / Duration.GetTotalSeconds());
+					OnBlockSecondaryIndex += LastLoopIdxE + 1;
+				}
+			}
+			else
+			{
+				if (WrappedModulo(LastTimeRange.GetLowerBoundValue(), Duration) < Start)
+				{
+					int32 LastLoopIdxS = (int32)FMath::FloorToInt(LastTimeRange.GetLowerBoundValue().GetTotalSeconds() / Duration.GetTotalSeconds());
+					OnBlockSecondaryIndex += LastLoopIdxS - (LoopIdxE + 1);
+				}
+			}
 		}
 		else
 		{
-			BlockOnRange = TRange<FMediaTimeStamp>(std::min(BlockOnRange.GetLowerBoundValue(), SetBlockOnRange.GetLowerBoundValue()), std::min(BlockOnRange.GetUpperBoundValue(), SetBlockOnRange.GetUpperBoundValue()));
+			// No old range data. We must assume this as the first block after startup / flush and start at loop index zero. Relocate indices by moving the base, so we really start at zero...
+			check(OnBlockSecondaryIndex == 0);
+			if (!bReverse)
+			{
+				OnBlockSecondaryIndex = -LoopIdxS;
+			}
+			else
+			{
+				OnBlockSecondaryIndex = -LoopIdxE;
+			}
 		}
+
+		// Assemble final blocking range
+		auto SeqIndexStart = FMediaTimeStamp::MakeSequenceIndex(OnBlockPrimaryIndex, OnBlockSecondaryIndex + LoopIdxS);
+		auto SeqIndexEnd = FMediaTimeStamp::MakeSequenceIndex(OnBlockPrimaryIndex, OnBlockSecondaryIndex + LoopIdxE);
+		BlockOnRange = TRange<FMediaTimeStamp>(FMediaTimeStamp(Start, SeqIndexStart), FMediaTimeStamp(End, SeqIndexEnd));
+		check(!BlockOnRange.IsEmpty());
 	}
 
+
+	// Note: Due to varying DTs the new range will NOT be a simple monotone progression in playback direction, but might overlap or even be a subset of the previous one
+	//		 We do not put any safeguards in place here, but rather use the "is last sample still valid" logic to reject illogical / impossible range requests.
+	//		 All that aside: we DO expect ranges start (lower bound if forward, upper if reverse playback) to be moving in a monotone manner according to the set playback direction.
+
 	CurrentPlayer->GetControls().SetBlockingPlaybackHint(!BlockOnRange.IsEmpty());
 
 	LastTimeRange = CurrentTimeRange;
@@ -1581,22 +1604,6 @@ bool FMediaPlayerFacade::BlockOnFetch() const
 					return false;
 				}
 			}
-			// Paused?
-			if (CurrentRate == 0.0f)
-			{
-				// Yes. If we get here the current sample did not satisfy the set range. So only a incoming seek could resolved this...
-				if (!SeekTargetTime.IsValid())
-				{
-					// No, so don't block. This range is impossible to satisfy...
-					return false;
-				}
-				// Would the seek target probably fit the range?
-				if (!BR.Overlaps(TRange<FMediaTimeStamp>(SeekTargetTime, SeekTargetTime + BR.Size<FMediaTimeStamp>().Time)))
-				{
-					// No, so don't block. This range is impossible to satisfy...
-					return false;
-				}
-			}
 		}
 
 		// Block until sample arrives!
@@ -1673,6 +1680,17 @@ void FMediaPlayerFacade::Flush(bool bExcludePlayer, bool bOnSeek)
 	{
 		LastVideoSampleProcessedTimeRange = TRange<FMediaTimeStamp>::Empty();
 	}
+	else
+	{
+		if (!bExcludePlayer && !LastVideoSampleProcessedTimeRange.IsEmpty())
+		{
+			// Players will reset their sequence index related values, but keep the playback position. Adjust our record accordingly...
+			int32 LoopIdxS = FMediaTimeStamp::GetSecondaryIndex(LastVideoSampleProcessedTimeRange.GetLowerBoundValue().SequenceIndex);
+			int32 LoopIdxE = FMediaTimeStamp::GetSecondaryIndex(LastVideoSampleProcessedTimeRange.GetUpperBoundValue().SequenceIndex);
+			LastVideoSampleProcessedTimeRange.SetLowerBoundValue(FMediaTimeStamp(LastVideoSampleProcessedTimeRange.GetLowerBoundValue().Time, FMediaTimeStamp::MakeSequenceIndex(0, 0)));
+			LastVideoSampleProcessedTimeRange.SetUpperBoundValue(FMediaTimeStamp(LastVideoSampleProcessedTimeRange.GetUpperBoundValue().Time, FMediaTimeStamp::MakeSequenceIndex(0, LoopIdxE - LoopIdxS)));
+		}
+	}
 
 	// Invalidate next video time to fetch (none-audio case)
 	NextEstVideoTimeAtFrameStart.Invalidate();
@@ -2600,7 +2618,7 @@ void FMediaPlayerFacade::PostSampleProcessingTimeHandling(FTimespan DeltaTime)
 					{
 						if (NextEstVideoTimeAtFrameStart.TimeStamp.Time >= Duration)
 						{
-							NextEstVideoTimeAtFrameStart.TimeStamp.Time = Duration - FTimespan::FromSeconds(0.0001);
+							NextEstVideoTimeAtFrameStart.TimeStamp.Time = Duration - FTimespan(1);
 						}
 					}
 					else
@@ -2755,12 +2773,13 @@ bool FMediaPlayerFacade::GetCurrentPlaybackTimeRange(TRange<FMediaTimeStamp>& Ti
 		return false;
 	}
 
+	const FTimespan Duration = Player->GetControls().GetDuration();
+
 	// If we are looping we check to prepare proper ranges should we wrap around either end of the media...
 	// (we do not clamp in the non-looping case as the rest of the code should deal with that fine)
 	if (Player->GetControls().IsLooping())
 	{
 		// Can the player already hand out a duration? (we might get here very early on)
-		const FTimespan Duration = Player->GetControls().GetDuration();
 		if (Duration != FTimespan::Zero())
 		{
 			FTimespan WrappedStart = WrappedModulo(TimeRange.GetLowerBoundValue().Time, Duration);
@@ -2778,6 +2797,11 @@ bool FMediaPlayerFacade::GetCurrentPlaybackTimeRange(TRange<FMediaTimeStamp>& Ti
 			}
 		}
 	}
+	else
+	{
+		TimeRange.SetLowerBoundValue(FMediaTimeStamp(FMath::Clamp(TimeRange.GetLowerBoundValue().Time, FTimespan::Zero(), Duration), TimeRange.GetLowerBoundValue().SequenceIndex));
+		TimeRange.SetUpperBoundValue(FMediaTimeStamp(FMath::Clamp(TimeRange.GetUpperBoundValue().Time, FTimespan::Zero(), Duration), TimeRange.GetUpperBoundValue().SequenceIndex));
+	}
 
 	return !TimeRange.IsEmpty();
 }
@@ -2786,32 +2810,6 @@ bool FMediaPlayerFacade::GetCurrentPlaybackTimeRange(TRange<FMediaTimeStamp>& Ti
 TRange<FMediaTimeStamp> FMediaPlayerFacade::GetAdjustedBlockOnRange() const
 {
 	TRange<FMediaTimeStamp> TimeRange = BlockOnRange.GetRange();
-
-	// The blocking range returned will base any secondary index values on zero. We need to patch in the secondary index value as last returned by the player plugin
-	if (!LastVideoSampleProcessedTimeRange.IsEmpty())
-	{
-		int32 BaseLoopIndex = FMediaTimeStamp::GetSecondaryIndex(LastVideoSampleProcessedTimeRange.GetLowerBoundValue().SequenceIndex);
-
-		// We still might have a time that is before the current one, in which case we need to push the loop index up (or down)...
-		if (GetUnpausedRate() >= 0.0f)
-		{
-			if (LastVideoSampleProcessedTimeRange.GetLowerBoundValue() > FMediaTimeStamp(TimeRange.GetLowerBoundValue().Time, FMediaTimeStamp::AdjustSecondaryIndex(TimeRange.GetLowerBoundValue().SequenceIndex, BaseLoopIndex)))
-			{
-				++BaseLoopIndex;
-			}
-		}
-		else
-		{
-			if (LastVideoSampleProcessedTimeRange.GetLowerBoundValue() < FMediaTimeStamp(TimeRange.GetLowerBoundValue().Time, FMediaTimeStamp::AdjustSecondaryIndex(TimeRange.GetLowerBoundValue().SequenceIndex, BaseLoopIndex)))
-			{
-				--BaseLoopIndex;
-			}
-		}
-
-		TimeRange.SetLowerBoundValue(FMediaTimeStamp(TimeRange.GetLowerBoundValue().Time, FMediaTimeStamp::AdjustSecondaryIndex(TimeRange.GetLowerBoundValue().SequenceIndex, BaseLoopIndex)));
-		TimeRange.SetUpperBoundValue(FMediaTimeStamp(TimeRange.GetUpperBoundValue().Time, FMediaTimeStamp::AdjustSecondaryIndex(TimeRange.GetUpperBoundValue().SequenceIndex, BaseLoopIndex)));
-	}
-
 	return TimeRange;
 }
 
@@ -2906,11 +2904,17 @@ bool FMediaPlayerFacade::IsVideoSampleStillGood(const TRange<FMediaTimeStamp>& L
 		// This better be true at all times
 		check(LastSampleTimeRange.GetLowerBoundValue().SequenceIndex == LastSampleTimeRange.GetUpperBoundValue().SequenceIndex);
 
+		// - - - - - - - - - - - - - - - - - - - - - - - - - - - - - - - - - - - - - - - - - - - - - - - - - - - - - - - - - - - - - - - - -
+		// Remap all values so we can assume all of them to be in a single "sequence index range" so the math doesn't get too unruly below
+
+		double Duration = Player->GetControls().GetDuration().GetTotalSeconds();
+
 		TRange<FMediaTimeStamp> TimeRange0;
 
 		// If we encounter a time range crossing some sequence index change, we need to check if we can "unroll" it...
 		uint64 LowerSeqIdx = TimeRange.GetLowerBoundValue().SequenceIndex;
 		uint64 UpperSeqIdx = TimeRange.GetUpperBoundValue().SequenceIndex;
+		uint64 RefSeqIdx = LowerSeqIdx;
 		if (LowerSeqIdx != UpperSeqIdx)
 		{
 			if (FMediaTimeStamp::GetPrimaryIndex(LowerSeqIdx) != FMediaTimeStamp::GetPrimaryIndex(UpperSeqIdx))
@@ -2919,13 +2923,15 @@ bool FMediaPlayerFacade::IsVideoSampleStillGood(const TRange<FMediaTimeStamp>& L
 				return false;
 			}
 
+			// We only should get here with a looping player
+			check(Player->GetControls().IsLooping());
+
 			// So we must have a loop index change. Compute how many loops and change the range into one "unrolled" one as indicated by the playback direction...
 			int32 LoopIdxDiff = FMediaTimeStamp::GetSecondaryIndex(UpperSeqIdx)
 							  - FMediaTimeStamp::GetSecondaryIndex(LowerSeqIdx);
+			// Note: this will be positive even with reverse playback as the orientation of the range will no change
 			check(LoopIdxDiff > 0);
 
-			double Duration = Player->GetControls().GetDuration().GetTotalSeconds();
-
 			if (!bReverse)
 			{
 				TimeRange0 = TRange<FMediaTimeStamp>(FMediaTimeStamp(TimeRange.GetLowerBoundValue().Time, 0), FMediaTimeStamp(TimeRange.GetUpperBoundValue().Time + FTimespan::FromSeconds(LoopIdxDiff * Duration), 0));
@@ -2933,31 +2939,55 @@ bool FMediaPlayerFacade::IsVideoSampleStillGood(const TRange<FMediaTimeStamp>& L
 			else
 			{
 				TimeRange0 = TRange<FMediaTimeStamp>(FMediaTimeStamp(TimeRange.GetLowerBoundValue().Time - FTimespan::FromSeconds(LoopIdxDiff * Duration), 0), FMediaTimeStamp(TimeRange.GetUpperBoundValue().Time, 0));
+				RefSeqIdx = UpperSeqIdx;
 			}
 		}
 		else
 		{
 			// Simple case, just bring everything down to "zero sequence index" for ease of processing below...
 			TimeRange0 = TRange<FMediaTimeStamp>(FMediaTimeStamp(TimeRange.GetLowerBoundValue().Time, 0), FMediaTimeStamp(TimeRange.GetUpperBoundValue().Time, 0));
+
+			// Is looping off?
+			if (!Player->GetControls().IsLooping())
+			{
+				// Yes. We clamp the range to the duration of the video to avoid looking at non-existent "next" frames...
+				TimeRange0 = TRange<FMediaTimeStamp>::Intersection(TimeRange0, TRange<FMediaTimeStamp>(FMediaTimeStamp(FTimespan::Zero(), 0), FMediaTimeStamp(FTimespan::FromSeconds(Duration), 0)));
+			}
 		}
 
-		TRange<FMediaTimeStamp> LastSampleTimeRange0(FMediaTimeStamp(LastSampleTimeRange.GetLowerBoundValue().Time, 0), FMediaTimeStamp(LastSampleTimeRange.GetUpperBoundValue().Time, 0));
+		// Map the last sample's time range to the same "sequence index range as the time range
+		int32 LastSampleIdxDiff = FMediaTimeStamp::GetSecondaryIndex(LastSampleTimeRange.GetLowerBoundValue().SequenceIndex) - RefSeqIdx;
+		FTimespan TimeOffset = FTimespan::FromSeconds(Duration * LastSampleIdxDiff);
+		TRange<FMediaTimeStamp> LastSampleTimeRange0(FMediaTimeStamp(LastSampleTimeRange.GetLowerBoundValue().Time + TimeOffset, 0), FMediaTimeStamp(LastSampleTimeRange.GetUpperBoundValue().Time + TimeOffset, 0));
+
+		// - - - - - - - - - - - - - - - - - - - - - - - - - - - - - - - - - - - - - - - - - - - - - - - - - - - - - - - - - - - -
+		// Now we can begin the checks with all time ranges mapped back to "zero sequence index"
+
+		// Is the sample time range ahead of the given time range?
+		// (did the range move in an unexpected way?)
+		if (!bReverse ? TimeRange0.GetUpperBoundValue() <= LastSampleTimeRange0.GetLowerBoundValue()
+					  : TimeRange0.GetLowerBoundValue() >= LastSampleTimeRange0.GetUpperBoundValue())
+		{
+			// We simply let the last sample stay around...
+			return true;
+		}
 
 		// Is the sample time range at all still valid?
 		if (LastSampleTimeRange0.Overlaps(TimeRange0))
 		{
-			// Yes. Assuming we could get more samples (of the same type) from the player, would the next one "better"?
-			// (this does not say which "next one" is actually current (e.g. if the rendering FPS is real bad), just that "another one" would be better)
+			// Yes. Assuming we could get more samples (of the same type) from the player, would the next one be "better"?
+			// (we assume samples of equal length)
 
 			// Compute the "theoretical" next sample range...
 			TRange<FMediaTimeStamp> NextSampleTimeRange = !bReverse ? TRange<FMediaTimeStamp>(LastSampleTimeRange0.GetUpperBoundValue(), LastSampleTimeRange0.GetUpperBoundValue() + LastSampleTimeRange0.Size<FMediaTimeStamp>().Time)
 																	: TRange<FMediaTimeStamp>(LastSampleTimeRange0.GetLowerBoundValue() - LastSampleTimeRange0.Size<FMediaTimeStamp>().Time, LastSampleTimeRange0.GetLowerBoundValue());
 
-			// Note: Loops (or the end of the timeline in non-looping setups)
+			// Note: Loops (or the end of the time line in non-looping setups)
 			// 
 			// - We could check for them and generate proper changes to the sequence index
 			// - Doing this would leave us with quite complex setups to compute the coverage
-			// - We opt for a cleaner, simpler approach: as we are NOT interested into proper PTS values, we can safely work with an "inifite" timeline when computing any overlaps, coverage and such
+			// - We opt for a cleaner, simpler approach: as we are NOT interested into proper PTS values, we can safely work with an "infinite" time line when computing any overlaps, coverage and such
+			//   (note: we DO need to restrict the range to the actual media duration if not looping - the code above does this)
 			// 
 			// --> we simply keep what we compute above!
 			//
@@ -2966,7 +2996,8 @@ bool FMediaPlayerFacade::IsVideoSampleStillGood(const TRange<FMediaTimeStamp>& L
 			int64 LastSampleCoverage = TRange<FMediaTimeStamp>::Intersection(TimeRange0, LastSampleTimeRange0).Size<FMediaTimeStamp>().Time.GetTicks();
 			int64 NextSampleCoverage = TRange<FMediaTimeStamp>::Intersection(TimeRange0, NextSampleTimeRange).Size<FMediaTimeStamp>().Time.GetTicks();
 
-			if (LastSampleCoverage > NextSampleCoverage)
+			// A new one is only desirable if it's BETTER than the current one
+			if (LastSampleCoverage >= NextSampleCoverage)
 			{
 				// Last one we returned is still good. No new one needed...
 				return true;
@@ -2979,13 +3010,16 @@ bool FMediaPlayerFacade::IsVideoSampleStillGood(const TRange<FMediaTimeStamp>& L
 
 bool FMediaPlayerFacade::ProcessVideoSamples(IMediaSamples& Samples, const TRange<FMediaTimeStamp>& TimeRange)
 {
-	// Let the player do some processing if needed.
-	if (Player.IsValid())
+	if (!Player.IsValid())
 	{
-		// note: avoid using this - it will be deprecated
-		Player->ProcessVideoSamples();
+		// Nothing to do, but in a sense: "successful"...
+		return true;
 	}
 
+	// Let the player do some processing if needed.
+	// note: avoid using this - it will be deprecated
+	Player->ProcessVideoSamples();
+
 	// This is not to be used with V1 timing
 	check(Player->GetPlayerFeatureFlag(IMediaPlayer::EFeatureFlag::UsePlaybackTimingV2));
 	// We expect a fully closed range or we assume: nothing to do...
@@ -3006,7 +3040,7 @@ bool FMediaPlayerFacade::ProcessVideoSamples(IMediaSamples& Samples, const TRang
 			return true;
 		}
 
-		switch (Samples.FetchBestVideoSampleForTimeRange(TimeRange, Sample, bReverse))
+		switch (Samples.FetchBestVideoSampleForTimeRange(TimeRange, Sample, bReverse, BlockOnRange.IsSet()))
 		{
 		case IMediaSamples::EFetchBestSampleResult::Ok:
 			break;
@@ -3056,8 +3090,7 @@ bool FMediaPlayerFacade::ProcessVideoSamples(IMediaSamples& Samples, const TRang
 	// Any sample?
 	if (Sample.IsValid())
 	{
-		// Yes. If we are in blocking playback mode we need to make sure that the sample is really in the range we asked for and block on...
-		// (same players might return an older sample as stop-gap measure if nothing can be found in the current range)
+		// Yes, deliver it and update state...
 
 		FMediaTimeStamp SampleTime = Sample->GetTime();
 		TRange<FMediaTimeStamp> SampleTimeRange(SampleTime, SampleTime + Sample->GetDuration());
diff --git a/Engine/Source/Runtime/MediaUtils/Private/MediaSamples.cpp b/Engine/Source/Runtime/MediaUtils/Private/MediaSamples.cpp
--- a/Engine/Source/Runtime/MediaUtils/Private/MediaSamples.cpp
+++ b/Engine/Source/Runtime/MediaUtils/Private/MediaSamples.cpp
@@ -139,9 +139,9 @@ void FMediaSamples::FlushSamples()
 /**
  * Fetch video sample best suited for the given time range. Samples prior to the selected one will be removed from the queue.
  */
-FMediaSamples::EFetchBestSampleResult FMediaSamples::FetchBestVideoSampleForTimeRange(const TRange<FMediaTimeStamp> & TimeRange, TSharedPtr<IMediaTextureSample, ESPMode::ThreadSafe>& OutSample, bool bReverse)
+FMediaSamples::EFetchBestSampleResult FMediaSamples::FetchBestVideoSampleForTimeRange(const TRange<FMediaTimeStamp> & TimeRange, TSharedPtr<IMediaTextureSample, ESPMode::ThreadSafe>& OutSample, bool bReverse, bool bConsistentResult)
 {
-	if (!VideoSampleQueue.FetchBestSampleForTimeRange(TimeRange, OutSample, bReverse))
+	if (!VideoSampleQueue.FetchBestSampleForTimeRange(TimeRange, OutSample, bReverse, bConsistentResult))
 	{
 		return EFetchBestSampleResult::NoSample;
 	}
diff --git a/Engine/Source/Runtime/MediaUtils/Public/MediaPlayerFacade.h b/Engine/Source/Runtime/MediaUtils/Public/MediaPlayerFacade.h
--- a/Engine/Source/Runtime/MediaUtils/Public/MediaPlayerFacade.h
+++ b/Engine/Source/Runtime/MediaUtils/Public/MediaPlayerFacade.h
@@ -908,6 +908,7 @@ private:
 			LastTimeRange = TRange<FTimespan>::Empty();
 			RangeIsDirty = false;
 			OnBlockPrimaryIndex = 0;
+			OnBlockSecondaryIndex = 0;
 		}
 
 	private:
@@ -926,8 +927,11 @@ private:
 		/** Flag to indicate if internal range is valid or not */
 		mutable bool RangeIsDirty;
 
-		/** Primary sequence index used during blocked playback processing */
+		/** Primary ("seek") sequence index used during blocked playback processing */
 		mutable int32 OnBlockPrimaryIndex;
+
+		/** Secondary (loop) sequence index used during blocked playback processing */
+		mutable int32 OnBlockSecondaryIndex;
 	};
 
 	FBlockOnRange BlockOnRange;
diff --git a/Engine/Source/Runtime/MediaUtils/Public/MediaSampleQueue.h b/Engine/Source/Runtime/MediaUtils/Public/MediaSampleQueue.h
--- a/Engine/Source/Runtime/MediaUtils/Public/MediaSampleQueue.h
+++ b/Engine/Source/Runtime/MediaUtils/Public/MediaSampleQueue.h
@@ -144,8 +144,13 @@ public:
 		return false;
 	}
 
-	bool FetchBestSampleForTimeRange(const TRange<FMediaTimeStamp>& TimeRange, TSharedPtr<SampleType, ESPMode::ThreadSafe>& OutSample, bool bReverse)
+	bool FetchBestSampleForTimeRange(const TRange<FMediaTimeStamp>& TimeRange, TSharedPtr<SampleType, ESPMode::ThreadSafe>& OutSample, bool bReverse, bool bConsistentResult)
 	{
+		// Notes:
+		// - Reverse playback still works with increasing indices in the queue. PTS values will be going down in it, rather than up,
+		//   but the order of indices is still identical.
+		// - The code below must be able to deal with time ranges that span loop points (different secondary sequence indices)
+		
 		// Code below assumes a fully specified range, no open bounds!
 		check(TimeRange.HasLowerBound() && TimeRange.HasUpperBound());
 
@@ -159,56 +164,86 @@ public:
 		// Found anything?
 		if (FirstPossibleIndex >= 0)
 		{
-			if (FirstPossibleIndex != LastPossibleIndex)
+			if (!bConsistentResult)
 			{
-				// More then one sample. Find the one that fits the bill, best...
-				// (we look for the one with the largest overlap & newest time)
-				FMediaTimeStamp BestDuration(FTimespan::Zero(), -1);
-				int32 BestIndex = FirstPossibleIndex;
-				for (int32 Idx = FirstPossibleIndex; Idx <= LastPossibleIndex; ++Idx)
+				//
+				// Return the latest sample with the most "coverage" for the given time range that can be fetched
+				// (this naturally depends on how many samples or available and hence timing - the result is not consistent between instances or repeat runs)
+				//
+				if (FirstPossibleIndex != LastPossibleIndex)
 				{
-					const TSharedPtr<SampleType, ESPMode::ThreadSafe>& Sample = Samples[Idx];
-
-					// Check once more if this sample is actually overlapping as we may get non-monotonically increasing data...
-					TRange<FMediaTimeStamp> SampleTimeRange = !bReverse ? TRange<FMediaTimeStamp>(Sample->GetTime(), Sample->GetTime() + Sample->GetDuration())
-						: TRange<FMediaTimeStamp>(Sample->GetTime() - Sample->GetDuration(), Sample->GetTime());
-
-					if (TimeRange.Overlaps(SampleTimeRange))
+					// More then one sample. Find the one that fits the bill, best...
+					// (we look for the one with the largest overlap & newest time)
+					FMediaTimeStamp BestDuration(FTimespan::Zero(), -1);
+					int32 BestIndex = FirstPossibleIndex;
+					for (int32 Idx = FirstPossibleIndex; Idx <= LastPossibleIndex; ++Idx)
 					{
-						// Ok. This one is real, see if it is a better fit than the last one...
-						TRange<FMediaTimeStamp> SampleInRangeRange(TRange<FMediaTimeStamp>::Intersection(SampleTimeRange, TimeRange));
+						const TSharedPtr<SampleType, ESPMode::ThreadSafe>& Sample = Samples[Idx];
+
+						// Check once more if this sample is actually overlapping as we may get non-monotonically increasing data...
+						TRange<FMediaTimeStamp> SampleTimeRange = TRange<FMediaTimeStamp>(Sample->GetTime(), Sample->GetTime() + Sample->GetDuration());
 
-						FMediaTimeStamp SampleDuration(SampleInRangeRange.Size<FMediaTimeStamp>());
-						if (SampleDuration >= BestDuration)
+						if (TimeRange.Overlaps(SampleTimeRange))
 						{
-							BestDuration = SampleDuration;
-							BestIndex = Idx;
+							// Ok. This one is real, see if it is a better fit than the last one...
+							TRange<FMediaTimeStamp> SampleInRangeRange(TRange<FMediaTimeStamp>::Intersection(SampleTimeRange, TimeRange));
+
+							FMediaTimeStamp SampleDuration(SampleInRangeRange.Size<FMediaTimeStamp>());
+							if (SampleDuration >= BestDuration)
+							{
+								BestDuration = SampleDuration;
+								BestIndex = Idx;
+							}
 						}
 					}
-				}
 
-				check(BestIndex >= NumOldSamplesAtBegin);
+					check(BestIndex >= NumOldSamplesAtBegin);
 
-				// Found the best. Return it & delete all candidate samples up and including it from the queue
-				OutSample = Samples[BestIndex];
-				Samples.RemoveAt(FirstPossibleIndex, BestIndex - FirstPossibleIndex + 1);
+					// Found the best. Return it & delete all candidate samples up and including it from the queue
+					OutSample = Samples[BestIndex];
+					Samples.RemoveAt(FirstPossibleIndex, BestIndex - FirstPossibleIndex + 1);
+				}
+				else
+				{
+					// Single sample found: we just take it!
+					OutSample = Samples[FirstPossibleIndex]; //-V781 PVS-Studio triggers incorrectly here: Variable checked after being used (likely a template code issue, but harmless)
+					Samples.RemoveAt(FirstPossibleIndex);
+				}
 			}
 			else
 			{
-				// Single sample found: we just take it!
-				OutSample = Samples[FirstPossibleIndex]; //-V781 PVS-Studio triggers incorrectly here: Variable checked after being used (likely a template code issue, but harmless)
-				Samples.RemoveAt(FirstPossibleIndex);
+				//
+				// Return the first sample with maximum possible coverage given the time range or nothing if the sample is not yet in the queue
+				// (this yields reproducible results between instances as far as the selection of frames is concerned if the passed in ranges are identical in each run / instance)
+				// 
+
+				for (int32 Idx = FirstPossibleIndex; Idx <= LastPossibleIndex; ++Idx)
+				{
+					const TSharedPtr<SampleType, ESPMode::ThreadSafe>& Sample = Samples[Idx];
+
+					TRange<FMediaTimeStamp> SampleTimeRange = TRange<FMediaTimeStamp>(Sample->GetTime(), Sample->GetTime() + Sample->GetDuration());
+
+					check(TimeRange.Overlaps(SampleTimeRange));
+					TRange<FMediaTimeStamp> SampleInRangeRange(TRange<FMediaTimeStamp>::Intersection(SampleTimeRange, TimeRange));
+
+					FMediaTimeStamp SampleDuration(SampleInRangeRange.Size<FMediaTimeStamp>());
+
+					// Do we either have full coverage or no "better" sample could be after this one and still in range?
+					if (!bReverse ? (SampleDuration.Time >= Sample->GetDuration() || (TimeRange.GetUpperBoundValue() - (Sample->GetTime() + Sample->GetDuration())) < SampleDuration)
+								  : (SampleDuration.Time >= Sample->GetDuration() || (Sample->GetTime() - TimeRange.GetLowerBoundValue()) < SampleDuration))
+					{
+						// Yes, so we return this one and remove all candidates and this one from the queue
+						OutSample = Sample;
+						Samples.RemoveAt(FirstPossibleIndex, Idx - FirstPossibleIndex + 1);
+						break;
+					}
+				}
 			}
 		}
 
 		// Any frames considered outdated?
 		if (NumOldSamplesAtBegin != 0)
 		{
-			// In case we got no new frame that fits into the current frame, return the newest of the "old" frames
-			if (!OutSample.IsValid())
-			{
-				OutSample = Samples[NumOldSamplesAtBegin - 1];
-			}
 			// Cleanup samples that are now considered outdated...
 			Samples.RemoveAt(0, NumOldSamplesAtBegin);
 		}
@@ -305,8 +340,7 @@ protected:
 			for (int32 Idx = 0; Idx < Num; ++Idx)
 			{
 				const TSharedPtr<SampleType, ESPMode::ThreadSafe>& Sample = Samples[Idx];
-				TRange<FMediaTimeStamp> SampleTimeRange = !bReverse ? TRange<FMediaTimeStamp>(Sample->GetTime(), Sample->GetTime() + Sample->GetDuration())
-					: TRange<FMediaTimeStamp>(Sample->GetTime() - Sample->GetDuration(), Sample->GetTime());
+				TRange<FMediaTimeStamp> SampleTimeRange = TRange<FMediaTimeStamp>(Sample->GetTime(), Sample->GetTime() + Sample->GetDuration());
 
 				if (TimeRange.Overlaps(SampleTimeRange))
 				{
@@ -320,7 +354,7 @@ protected:
 				else
 				{
 					if (!bReverse ? (SampleTimeRange.GetLowerBoundValue() >= TimeRange.GetUpperBoundValue()) :
-						(SampleTimeRange.GetUpperBoundValue() <= TimeRange.GetLowerBoundValue()))
+									(SampleTimeRange.GetUpperBoundValue() <= TimeRange.GetLowerBoundValue()))
 					{
 						// Sample is entirely past requested time range, we can stop
 						// (we assume monotonically increasing time stamps here)
diff --git a/Engine/Source/Runtime/MediaUtils/Public/MediaSamples.h b/Engine/Source/Runtime/MediaUtils/Public/MediaSamples.h
--- a/Engine/Source/Runtime/MediaUtils/Public/MediaSamples.h
+++ b/Engine/Source/Runtime/MediaUtils/Public/MediaSamples.h
@@ -163,7 +163,7 @@ public:
 	MEDIAUTILS_API virtual bool FetchSubtitle(TRange<FMediaTimeStamp> TimeRange, TSharedPtr<IMediaOverlaySample, ESPMode::ThreadSafe>& OutSample) override;
 	MEDIAUTILS_API virtual bool FetchVideo(TRange<FMediaTimeStamp> TimeRange, TSharedPtr<IMediaTextureSample, ESPMode::ThreadSafe>& OutSample) override;
 
-	MEDIAUTILS_API virtual EFetchBestSampleResult FetchBestVideoSampleForTimeRange(const TRange<FMediaTimeStamp> & TimeRange, TSharedPtr<IMediaTextureSample, ESPMode::ThreadSafe>& OutSample, bool bReverse) override;
+	MEDIAUTILS_API virtual EFetchBestSampleResult FetchBestVideoSampleForTimeRange(const TRange<FMediaTimeStamp>& TimeRange, TSharedPtr<IMediaTextureSample, ESPMode::ThreadSafe>& OutSample, bool bReverse, bool bConsistentResult) override;
 
 	/**
 	 * Peek next video sample's timestamp
