Fix for live playback of video streams in MediaFrameWorks (code was missing some handling for streams with no known duration)
#rnx
#tests Local machine (PC) with stream with and without audio track

diff --git a/Engine/Source/Runtime/MediaUtils/Private/MediaPlayerFacade.cpp b/Engine/Source/Runtime/MediaUtils/Private/MediaPlayerFacade.cpp
--- a/Engine/Source/Runtime/MediaUtils/Private/MediaPlayerFacade.cpp
+++ b/Engine/Source/Runtime/MediaUtils/Private/MediaPlayerFacade.cpp
@@ -1170,19 +1170,22 @@ bool FMediaPlayerFacade::Seek(const FTimespan& InTime)
 	}
 
 	FTimespan Duration = CurrentPlayer->GetControls().GetDuration();
-	if (Duration == FTimespan::Zero())
-	{
-		return false;
-	}
 
 	FTimespan Time;
-	if (CurrentPlayer->GetControls().IsLooping())
+	if (Duration > FTimespan::Zero())
 	{
-		Time = WrappedModulo(InTime, Duration);
+		if (CurrentPlayer->GetControls().IsLooping())
+		{
+			Time = WrappedModulo(InTime, Duration);
+		}
+		else
+		{
+			Time = FTimespan(FMath::Clamp(InTime.GetTicks(), (int64)0L, Duration.GetTicks()));
+		}
 	}
 	else
 	{
-		Time = FTimespan(FMath::Clamp(InTime.GetTicks(), (int64)0L, Duration.GetTicks()));
+		Time = InTime;
 	}
 
 	if (!CurrentPlayer->GetControls().Seek(Time))
@@ -2595,20 +2598,23 @@ void FMediaPlayerFacade::PostSampleProcessingTimeHandling(FTimespan DeltaTime)
 
 				if (Player->GetControls().IsLooping())
 				{
-					if (Rate >= 0.0f)
+					if (Duration > FTimespan::Zero())
 					{
-						while (NextEstVideoTimeAtFrameStart.TimeStamp.Time >= Duration)
+						if (Rate >= 0.0f)
 						{
-							NextEstVideoTimeAtFrameStart.TimeStamp.Time -= Duration;
-							NextEstVideoTimeAtFrameStart.TimeStamp.SequenceIndex = FMediaTimeStamp::AdjustSecondaryIndex(NextEstVideoTimeAtFrameStart.TimeStamp.SequenceIndex, 1);
+							while (NextEstVideoTimeAtFrameStart.TimeStamp.Time >= Duration)
+							{
+								NextEstVideoTimeAtFrameStart.TimeStamp.Time -= Duration;
+								NextEstVideoTimeAtFrameStart.TimeStamp.SequenceIndex = FMediaTimeStamp::AdjustSecondaryIndex(NextEstVideoTimeAtFrameStart.TimeStamp.SequenceIndex, 1);
+							}
 						}
-					}
-					else
-					{
-						while (NextEstVideoTimeAtFrameStart.TimeStamp.Time < FTimespan::Zero())
+						else
 						{
-							NextEstVideoTimeAtFrameStart.TimeStamp.Time += Duration;
-							NextEstVideoTimeAtFrameStart.TimeStamp.SequenceIndex = FMediaTimeStamp::AdjustSecondaryIndex(NextEstVideoTimeAtFrameStart.TimeStamp.SequenceIndex, -1);
+							while (NextEstVideoTimeAtFrameStart.TimeStamp.Time < FTimespan::Zero())
+							{
+								NextEstVideoTimeAtFrameStart.TimeStamp.Time += Duration;
+								NextEstVideoTimeAtFrameStart.TimeStamp.SequenceIndex = FMediaTimeStamp::AdjustSecondaryIndex(NextEstVideoTimeAtFrameStart.TimeStamp.SequenceIndex, -1);
+							}
 						}
 					}
 				}
@@ -2616,9 +2622,12 @@ void FMediaPlayerFacade::PostSampleProcessingTimeHandling(FTimespan DeltaTime)
 				{
 					if (Rate >= 0.0f)
 					{
-						if (NextEstVideoTimeAtFrameStart.TimeStamp.Time >= Duration)
+						if (Duration > FTimespan::Zero())
 						{
-							NextEstVideoTimeAtFrameStart.TimeStamp.Time = Duration - FTimespan(1);
+							if (NextEstVideoTimeAtFrameStart.TimeStamp.Time >= Duration)
+							{
+								NextEstVideoTimeAtFrameStart.TimeStamp.Time = Duration - FTimespan(1);
+							}
 						}
 					}
 					else
@@ -2775,12 +2784,12 @@ bool FMediaPlayerFacade::GetCurrentPlaybackTimeRange(TRange<FMediaTimeStamp>& Ti
 
 	const FTimespan Duration = Player->GetControls().GetDuration();
 
-	// If we are looping we check to prepare proper ranges should we wrap around either end of the media...
-	// (we do not clamp in the non-looping case as the rest of the code should deal with that fine)
-	if (Player->GetControls().IsLooping())
+	// We need a valid duration for the next steps (we may not have one e.g. for live material)
+	if (Duration > FTimespan::Zero())
 	{
-		// Can the player already hand out a duration? (we might get here very early on)
-		if (Duration != FTimespan::Zero())
+		// If we are looping we check to prepare proper ranges should we wrap around either end of the media...
+		// (we do not clamp in the non-looping case as the rest of the code should deal with that fine)
+		if (Player->GetControls().IsLooping())
 		{
 			FTimespan WrappedStart = WrappedModulo(TimeRange.GetLowerBoundValue().Time, Duration);
 			FTimespan WrappedEnd = WrappedModulo(TimeRange.GetUpperBoundValue().Time, Duration);
@@ -2796,11 +2805,11 @@ bool FMediaPlayerFacade::GetCurrentPlaybackTimeRange(TRange<FMediaTimeStamp>& Ti
 				}
 			}
 		}
-	}
-	else
-	{
-		TimeRange.SetLowerBoundValue(FMediaTimeStamp(FMath::Clamp(TimeRange.GetLowerBoundValue().Time, FTimespan::Zero(), Duration), TimeRange.GetLowerBoundValue().SequenceIndex));
-		TimeRange.SetUpperBoundValue(FMediaTimeStamp(FMath::Clamp(TimeRange.GetUpperBoundValue().Time, FTimespan::Zero(), Duration), TimeRange.GetUpperBoundValue().SequenceIndex));
+		else
+		{
+			TimeRange.SetLowerBoundValue(FMediaTimeStamp(FMath::Clamp(TimeRange.GetLowerBoundValue().Time, FTimespan::Zero(), Duration), TimeRange.GetLowerBoundValue().SequenceIndex));
+			TimeRange.SetUpperBoundValue(FMediaTimeStamp(FMath::Clamp(TimeRange.GetUpperBoundValue().Time, FTimespan::Zero(), Duration), TimeRange.GetUpperBoundValue().SequenceIndex));
+		}
 	}
 
 	return !TimeRange.IsEmpty();
@@ -2923,8 +2932,9 @@ bool FMediaPlayerFacade::IsVideoSampleStillGood(const TRange<FMediaTimeStamp>& L
 				return false;
 			}
 
-			// We only should get here with a looping player
+			// We only should get here with a looping player that knows its duration
 			check(Player->GetControls().IsLooping());
+			check(Duration > 0.0);
 
 			// So we must have a loop index change. Compute how many loops and change the range into one "unrolled" one as indicated by the playback direction...
 			int32 LoopIdxDiff = FMediaTimeStamp::GetSecondaryIndex(UpperSeqIdx)
@@ -2950,12 +2960,16 @@ bool FMediaPlayerFacade::IsVideoSampleStillGood(const TRange<FMediaTimeStamp>& L
 			// Is looping off?
 			if (!Player->GetControls().IsLooping())
 			{
-				// Yes. We clamp the range to the duration of the video to avoid looking at non-existent "next" frames...
-				TimeRange0 = TRange<FMediaTimeStamp>::Intersection(TimeRange0, TRange<FMediaTimeStamp>(FMediaTimeStamp(FTimespan::Zero(), 0), FMediaTimeStamp(FTimespan::FromSeconds(Duration), 0)));
+				// Yes. We clamp the range to the duration of the video to avoid looking at non-existent "next" frames... (unless we have no duration)
+				if (Duration > 0.0)
+				{
+					TimeRange0 = TRange<FMediaTimeStamp>::Intersection(TimeRange0, TRange<FMediaTimeStamp>(FMediaTimeStamp(FTimespan::Zero(), 0), FMediaTimeStamp(FTimespan::FromSeconds(Duration), 0)));
+				}
 			}
 		}
 
 		// Map the last sample's time range to the same "sequence index range as the time range
+		// (note: for e.g. life streams that do not have any set duration all this will not change the timerange - just as needed)
 		int32 LastSampleIdxDiff = FMediaTimeStamp::GetSecondaryIndex(LastSampleTimeRange.GetLowerBoundValue().SequenceIndex) - RefSeqIdx;
 		FTimespan TimeOffset = FTimespan::FromSeconds(Duration * LastSampleIdxDiff);
 		TRange<FMediaTimeStamp> LastSampleTimeRange0(FMediaTimeStamp(LastSampleTimeRange.GetLowerBoundValue().Time + TimeOffset, 0), FMediaTimeStamp(LastSampleTimeRange.GetUpperBoundValue().Time + TimeOffset, 0));
