Reworked version of live playback fix for MediaFramework (now also accepting "inifnite" durations as value)
#rnx
#tests Tested with lifestream on local machine

diff --git a/Engine/Source/Runtime/MediaUtils/Private/MediaPlayerFacade.cpp b/Engine/Source/Runtime/MediaUtils/Private/MediaPlayerFacade.cpp
--- a/Engine/Source/Runtime/MediaUtils/Private/MediaPlayerFacade.cpp
+++ b/Engine/Source/Runtime/MediaUtils/Private/MediaPlayerFacade.cpp
@@ -105,6 +105,11 @@ static FTimespan WrappedModulo(FTimespan Time, FTimespan Duration)
 	return (Time >= FTimespan::Zero()) ? (Time % Duration) : (Duration + (Time % Duration));
 }
 
+static bool IsDurationValidAndFinite(FTimespan Duration)
+{
+	return (Duration != FTimespan::Zero() && Duration.GetTicks() != TNumericLimits<int64>::Max());
+}
+
 /* FMediaPlayerFacade structors
 *****************************************************************************/
 
@@ -1172,7 +1177,7 @@ bool FMediaPlayerFacade::Seek(const FTimespan& InTime)
 	FTimespan Duration = CurrentPlayer->GetControls().GetDuration();
 
 	FTimespan Time;
-	if (Duration > FTimespan::Zero())
+	if (IsDurationValidAndFinite(Duration))
 	{
 		if (CurrentPlayer->GetControls().IsLooping())
 		{
@@ -1348,8 +1353,8 @@ const TRange<FMediaTimeStamp>& FMediaPlayerFacade::FBlockOnRange::GetRange() con
 		 only if data is processed, which would also mean: we know the duration!
 		 (Exception: live playback! --> but we would not allow blocking there anyway! (makes no sense as real life use case))
 		*/
-		check(!Duration.IsZero());
-		if (Duration.IsZero())
+		check(IsDurationValidAndFinite(Duration));
+		if (!IsDurationValidAndFinite(Duration))
 		{
 			// Catch if this is called to early and reset blocking...
 			BlockOnRange = TRange<FMediaTimeStamp>::Empty();
@@ -2598,7 +2603,7 @@ void FMediaPlayerFacade::PostSampleProcessingTimeHandling(FTimespan DeltaTime)
 
 				if (Player->GetControls().IsLooping())
 				{
-					if (Duration > FTimespan::Zero())
+					if (IsDurationValidAndFinite(Duration))
 					{
 						if (Rate >= 0.0f)
 						{
@@ -2622,7 +2627,7 @@ void FMediaPlayerFacade::PostSampleProcessingTimeHandling(FTimespan DeltaTime)
 				{
 					if (Rate >= 0.0f)
 					{
-						if (Duration > FTimespan::Zero())
+						if (IsDurationValidAndFinite(Duration))
 						{
 							if (NextEstVideoTimeAtFrameStart.TimeStamp.Time >= Duration)
 							{
@@ -2785,7 +2790,7 @@ bool FMediaPlayerFacade::GetCurrentPlaybackTimeRange(TRange<FMediaTimeStamp>& Ti
 	const FTimespan Duration = Player->GetControls().GetDuration();
 
 	// We need a valid duration for the next steps (we may not have one e.g. for live material)
-	if (Duration > FTimespan::Zero())
+	if (IsDurationValidAndFinite(Duration))
 	{
 		// If we are looping we check to prepare proper ranges should we wrap around either end of the media...
 		// (we do not clamp in the non-looping case as the rest of the code should deal with that fine)
@@ -2916,7 +2921,7 @@ bool FMediaPlayerFacade::IsVideoSampleStillGood(const TRange<FMediaTimeStamp>& L
 		// - - - - - - - - - - - - - - - - - - - - - - - - - - - - - - - - - - - - - - - - - - - - - - - - - - - - - - - - - - - - - - - - -
 		// Remap all values so we can assume all of them to be in a single "sequence index range" so the math doesn't get too unruly below
 
-		double Duration = Player->GetControls().GetDuration().GetTotalSeconds();
+		FTimespan Duration = Player->GetControls().GetDuration();
 
 		TRange<FMediaTimeStamp> TimeRange0;
 
@@ -2934,7 +2939,7 @@ bool FMediaPlayerFacade::IsVideoSampleStillGood(const TRange<FMediaTimeStamp>& L
 
 			// We only should get here with a looping player that knows its duration
 			check(Player->GetControls().IsLooping());
-			check(Duration > 0.0);
+			check(IsDurationValidAndFinite(Duration));
 
 			// So we must have a loop index change. Compute how many loops and change the range into one "unrolled" one as indicated by the playback direction...
 			int32 LoopIdxDiff = FMediaTimeStamp::GetSecondaryIndex(UpperSeqIdx)
@@ -2942,13 +2947,15 @@ bool FMediaPlayerFacade::IsVideoSampleStillGood(const TRange<FMediaTimeStamp>& L
 			// Note: this will be positive even with reverse playback as the orientation of the range will no change
 			check(LoopIdxDiff > 0);
 
+			double DurationD = Duration.GetTotalSeconds();
+
 			if (!bReverse)
 			{
-				TimeRange0 = TRange<FMediaTimeStamp>(FMediaTimeStamp(TimeRange.GetLowerBoundValue().Time, 0), FMediaTimeStamp(TimeRange.GetUpperBoundValue().Time + FTimespan::FromSeconds(LoopIdxDiff * Duration), 0));
+				TimeRange0 = TRange<FMediaTimeStamp>(FMediaTimeStamp(TimeRange.GetLowerBoundValue().Time, 0), FMediaTimeStamp(TimeRange.GetUpperBoundValue().Time + FTimespan::FromSeconds(LoopIdxDiff * DurationD), 0));
 			}
 			else
 			{
-				TimeRange0 = TRange<FMediaTimeStamp>(FMediaTimeStamp(TimeRange.GetLowerBoundValue().Time - FTimespan::FromSeconds(LoopIdxDiff * Duration), 0), FMediaTimeStamp(TimeRange.GetUpperBoundValue().Time, 0));
+				TimeRange0 = TRange<FMediaTimeStamp>(FMediaTimeStamp(TimeRange.GetLowerBoundValue().Time - FTimespan::FromSeconds(LoopIdxDiff * DurationD), 0), FMediaTimeStamp(TimeRange.GetUpperBoundValue().Time, 0));
 				RefSeqIdx = UpperSeqIdx;
 			}
 		}
@@ -2961,9 +2968,9 @@ bool FMediaPlayerFacade::IsVideoSampleStillGood(const TRange<FMediaTimeStamp>& L
 			if (!Player->GetControls().IsLooping())
 			{
 				// Yes. We clamp the range to the duration of the video to avoid looking at non-existent "next" frames... (unless we have no duration)
-				if (Duration > 0.0)
+				if (IsDurationValidAndFinite(Duration))
 				{
-					TimeRange0 = TRange<FMediaTimeStamp>::Intersection(TimeRange0, TRange<FMediaTimeStamp>(FMediaTimeStamp(FTimespan::Zero(), 0), FMediaTimeStamp(FTimespan::FromSeconds(Duration), 0)));
+					TimeRange0 = TRange<FMediaTimeStamp>::Intersection(TimeRange0, TRange<FMediaTimeStamp>(FMediaTimeStamp(FTimespan::Zero(), 0), FMediaTimeStamp(Duration, 0)));
 				}
 			}
 		}
@@ -2971,7 +2978,7 @@ bool FMediaPlayerFacade::IsVideoSampleStillGood(const TRange<FMediaTimeStamp>& L
 		// Map the last sample's time range to the same "sequence index range as the time range
 		// (note: for e.g. life streams that do not have any set duration all this will not change the timerange - just as needed)
 		int32 LastSampleIdxDiff = FMediaTimeStamp::GetSecondaryIndex(LastSampleTimeRange.GetLowerBoundValue().SequenceIndex) - RefSeqIdx;
-		FTimespan TimeOffset = FTimespan::FromSeconds(Duration * LastSampleIdxDiff);
+		FTimespan TimeOffset = IsDurationValidAndFinite(Duration) ? Duration * LastSampleIdxDiff : FTimespan::Zero();
 		TRange<FMediaTimeStamp> LastSampleTimeRange0(FMediaTimeStamp(LastSampleTimeRange.GetLowerBoundValue().Time + TimeOffset, 0), FMediaTimeStamp(LastSampleTimeRange.GetUpperBoundValue().Time + TimeOffset, 0));
 
 		// - - - - - - - - - - - - - - - - - - - - - - - - - - - - - - - - - - - - - - - - - - - - - - - - - - - - - - - - - - - -
