MFW: fixing playback issues with media that has a valid duration value, but does not start at PTS==0
#rnx
#tests Local machine test with VOD material offset via URL parameter

diff --git a/Engine/Source/Runtime/MediaUtils/Private/MediaPlayerFacade.cpp b/Engine/Source/Runtime/MediaUtils/Private/MediaPlayerFacade.cpp
--- a/Engine/Source/Runtime/MediaUtils/Private/MediaPlayerFacade.cpp
+++ b/Engine/Source/Runtime/MediaUtils/Private/MediaPlayerFacade.cpp
@@ -35,6 +35,7 @@
 
 #define MEDIAPLAYERFACADE_DISABLE_BLOCKING 0
 #define MEDIAPLAYERFACADE_TRACE_SINKOVERFLOWS 0
+#define MEDIAPLAYERFACADE_DISABLE_PTSCLAMP 1			// enable to disable any clamping of PTS values to the [0..duration[ range (allowing non-zero-PTS-based material with a duration to play)
 
 
 /** Time spent in media player facade closing media. */
@@ -1185,7 +1186,9 @@ bool FMediaPlayerFacade::Seek(const FTimespan& InTime)
 		}
 		else
 		{
+#if !MEDIAPLAYERFACADE_DISABLE_PTSCLAMP
 			Time = FTimespan(FMath::Clamp(InTime.GetTicks(), (int64)0L, Duration.GetTicks()));
+#endif
 		}
 	}
 	else
@@ -2812,8 +2815,10 @@ bool FMediaPlayerFacade::GetCurrentPlaybackTimeRange(TRange<FMediaTimeStamp>& Ti
 		}
 		else
 		{
+#if !MEDIAPLAYERFACADE_DISABLE_PTSCLAMP
 			TimeRange.SetLowerBoundValue(FMediaTimeStamp(FMath::Clamp(TimeRange.GetLowerBoundValue().Time, FTimespan::Zero(), Duration), TimeRange.GetLowerBoundValue().SequenceIndex));
 			TimeRange.SetUpperBoundValue(FMediaTimeStamp(FMath::Clamp(TimeRange.GetUpperBoundValue().Time, FTimespan::Zero(), Duration), TimeRange.GetUpperBoundValue().SequenceIndex));
+#endif
 		}
 	}
 
@@ -2967,11 +2972,13 @@ bool FMediaPlayerFacade::IsVideoSampleStillGood(const TRange<FMediaTimeStamp>& L
 			// Is looping off?
 			if (!Player->GetControls().IsLooping())
 			{
+#if !MEDIAPLAYERFACADE_DISABLE_PTSCLAMP
 				// Yes. We clamp the range to the duration of the video to avoid looking at non-existent "next" frames... (unless we have no duration)
 				if (IsDurationValidAndFinite(Duration))
 				{
 					TimeRange0 = TRange<FMediaTimeStamp>::Intersection(TimeRange0, TRange<FMediaTimeStamp>(FMediaTimeStamp(FTimespan::Zero(), 0), FMediaTimeStamp(Duration, 0)));
 				}
+#endif
 			}
 		}
 
