Fixing up recent regression causing seeks for most commonly used media not to work anymore
#rnx

diff --git a/Engine/Source/Runtime/MediaUtils/Private/MediaPlayerFacade.cpp b/Engine/Source/Runtime/MediaUtils/Private/MediaPlayerFacade.cpp
--- a/Engine/Source/Runtime/MediaUtils/Private/MediaPlayerFacade.cpp
+++ b/Engine/Source/Runtime/MediaUtils/Private/MediaPlayerFacade.cpp
@@ -1188,6 +1188,8 @@ bool FMediaPlayerFacade::Seek(const FTimespan& InTime)
 		{
 #if !MEDIAPLAYERFACADE_DISABLE_PTSCLAMP
 			Time = FTimespan(FMath::Clamp(InTime.GetTicks(), (int64)0L, Duration.GetTicks()));
+#else
+			Time = InTime;
 #endif
 		}
 	}
