Fixing up a logic issue that made some setups with MFW players not return samples for a while (around range-internal loop points) and posisbly could cause tiimeouts
#tests Tested problemnatic cases on local machine in test project

diff --git a/Engine/Source/Runtime/MediaUtils/Private/MediaPlayerFacade.cpp b/Engine/Source/Runtime/MediaUtils/Private/MediaPlayerFacade.cpp
--- a/Engine/Source/Runtime/MediaUtils/Private/MediaPlayerFacade.cpp
+++ b/Engine/Source/Runtime/MediaUtils/Private/MediaPlayerFacade.cpp
@@ -1264,7 +1264,7 @@ void FMediaPlayerFacade::FBlockOnRange::OnFlush()
 {
 	LastTimeRange = TRange<FTimespan>::Empty();
 	OnBlockPrimaryIndex = 0;
-	OnBlockSecondaryIndex = 0;
+	OnBlockSecondaryIndexOffset = 0;
 	RangeIsDirty = true;
 }
 
@@ -1273,7 +1273,7 @@ void FMediaPlayerFacade::FBlockOnRange::OnSeek(int32 PrimaryIndex)
 {
 	LastTimeRange = TRange<FTimespan>::Empty();
 	OnBlockPrimaryIndex = PrimaryIndex;
-	OnBlockSecondaryIndex = 0;
+	OnBlockSecondaryIndexOffset = 0;
 	RangeIsDirty = true;
 }
 
@@ -1333,8 +1333,6 @@ const TRange<FMediaTimeStamp>& FMediaPlayerFacade::FBlockOnRange::GetRange() con
 	FTimespan Start(CurrentTimeRange.GetLowerBoundValue());
 	FTimespan End(CurrentTimeRange.GetUpperBoundValue());
 
-	auto SetBlockOnRange = BlockOnRange;
-
 	/*
 	* On the synthesized sequence and loop index values:
 	* - We track seeks and hence can insert the proper seek index easily, although the user does not provide it / does not need to track it
@@ -1347,7 +1345,8 @@ const TRange<FMediaTimeStamp>& FMediaPlayerFacade::FBlockOnRange::GetRange() con
 
 	if (!CurrentPlayer->GetControls().IsLooping())
 	{
-		int64 SequenceIndex = FMediaTimeStamp::MakeSequenceIndex(OnBlockPrimaryIndex, OnBlockSecondaryIndex);
+		int32 LastLoopIdxS = LastTimeRange.IsEmpty() ? 0 : (int32)FMath::FloorToInt(LastTimeRange.GetLowerBoundValue().GetTotalSeconds() / Duration.GetTotalSeconds());
+		int64 SequenceIndex = FMediaTimeStamp::MakeSequenceIndex(OnBlockPrimaryIndex, OnBlockSecondaryIndexOffset + LastLoopIdxS);
 		BlockOnRange = TRange<FMediaTimeStamp>(FMediaTimeStamp(Start, SequenceIndex), FMediaTimeStamp(End, SequenceIndex));
 	}
 	else
@@ -1383,38 +1382,46 @@ const TRange<FMediaTimeStamp>& FMediaPlayerFacade::FBlockOnRange::GetRange() con
 			// Adjust loop index base such that, given a playback direction, we can guarantee that the new indices returned fit a monotone progression
 			if (!bReverse)
 			{
-				if (WrappedModulo(LastTimeRange.GetLowerBoundValue(), Duration) > Start)
+				// Did the range passed in, loop around?
+				if (LastTimeRange.GetLowerBoundValue() > CurrentTimeRange.GetLowerBoundValue())
 				{
-					int32 LastLoopIdxE = (int32)FMath::FloorToInt(LastTimeRange.GetUpperBoundValue().GetTotalSeconds() / Duration.GetTotalSeconds());
-					OnBlockSecondaryIndex += LastLoopIdxE + 1;
+					// Yes. Adjust the base secondary index value to guarantee a simple continuation of the secondary index values
+					int32 LastLoopIdxS = (int32)FMath::FloorToInt(LastTimeRange.GetLowerBoundValue().GetTotalSeconds() / Duration.GetTotalSeconds());
+					// -LoopIdxS -> Compensate for new start index
+					// +LastLoopIdxS + 1 -> Move index one beyond last
+					OnBlockSecondaryIndexOffset = -LoopIdxS + LastLoopIdxS + 1;
 				}
 			}
 			else
 			{
-				if (WrappedModulo(LastTimeRange.GetLowerBoundValue(), Duration) < Start)
+				// Did the range passed in, loop around?
+				if (LastTimeRange.GetLowerBoundValue() < CurrentTimeRange.GetLowerBoundValue())
 				{
+					// Yes. Adjust the base secondary index value to guarantee a simple continuation of the secondary index values
 					int32 LastLoopIdxS = (int32)FMath::FloorToInt(LastTimeRange.GetLowerBoundValue().GetTotalSeconds() / Duration.GetTotalSeconds());
-					OnBlockSecondaryIndex += LastLoopIdxS - (LoopIdxE + 1);
+					// -LoopIdxS -> Compensate for new start index
+					// +LastLoopIdxS + 1 -> Move index one beyond last 
+					OnBlockSecondaryIndexOffset = -LoopIdxS + LastLoopIdxS - 1;
 				}
 			}
 		}
 		else
 		{
 			// No old range data. We must assume this as the first block after startup / flush and start at loop index zero. Relocate indices by moving the base, so we really start at zero...
-			check(OnBlockSecondaryIndex == 0);
+			check(OnBlockSecondaryIndexOffset == 0);
 			if (!bReverse)
 			{
-				OnBlockSecondaryIndex = -LoopIdxS;
+				OnBlockSecondaryIndexOffset = -LoopIdxS;
 			}
 			else
 			{
-				OnBlockSecondaryIndex = -LoopIdxE;
+				OnBlockSecondaryIndexOffset = -LoopIdxE;
 			}
 		}
 
 		// Assemble final blocking range
-		auto SeqIndexStart = FMediaTimeStamp::MakeSequenceIndex(OnBlockPrimaryIndex, OnBlockSecondaryIndex + LoopIdxS);
-		auto SeqIndexEnd = FMediaTimeStamp::MakeSequenceIndex(OnBlockPrimaryIndex, OnBlockSecondaryIndex + LoopIdxE);
+		auto SeqIndexStart = FMediaTimeStamp::MakeSequenceIndex(OnBlockPrimaryIndex, OnBlockSecondaryIndexOffset + LoopIdxS);
+		auto SeqIndexEnd = FMediaTimeStamp::MakeSequenceIndex(OnBlockPrimaryIndex, OnBlockSecondaryIndexOffset + LoopIdxE);
 		BlockOnRange = TRange<FMediaTimeStamp>(FMediaTimeStamp(Start, SeqIndexStart), FMediaTimeStamp(End, SeqIndexEnd));
 		check(!BlockOnRange.IsEmpty());
 	}
diff --git a/Engine/Source/Runtime/MediaUtils/Public/MediaPlayerFacade.h b/Engine/Source/Runtime/MediaUtils/Public/MediaPlayerFacade.h
--- a/Engine/Source/Runtime/MediaUtils/Public/MediaPlayerFacade.h
+++ b/Engine/Source/Runtime/MediaUtils/Public/MediaPlayerFacade.h
@@ -908,7 +908,7 @@ private:
 			LastTimeRange = TRange<FTimespan>::Empty();
 			RangeIsDirty = false;
 			OnBlockPrimaryIndex = 0;
-			OnBlockSecondaryIndex = 0;
+			OnBlockSecondaryIndexOffset = 0;
 		}
 
 	private:
@@ -930,8 +930,8 @@ private:
 		/** Primary ("seek") sequence index used during blocked playback processing */
 		mutable int32 OnBlockPrimaryIndex;
 
-		/** Secondary (loop) sequence index used during blocked playback processing */
-		mutable int32 OnBlockSecondaryIndex;
+		/** Secondary (loop) sequence index offset used during blocked playback processing */
+		mutable int32 OnBlockSecondaryIndexOffset;
 	};
 
 	FBlockOnRange BlockOnRange;
