Media: Allow peeking at the time range covered by the enqueued media samples.

#rnx

diff --git a/Engine/Source/Runtime/Media/Public/IMediaSamples.h b/Engine/Source/Runtime/Media/Public/IMediaSamples.h
--- a/Engine/Source/Runtime/Media/Public/IMediaSamples.h
+++ b/Engine/Source/Runtime/Media/Public/IMediaSamples.h
@@ -127,7 +127,8 @@ public:
 		return EFetchBestSampleResult::NotSupported;
 	}
 
-	virtual bool PeekVideoSampleTime(FMediaTimeStamp & TimeStamp) = 0;
+	virtual bool PeekVideoSampleTime(FMediaTimeStamp& TimeStamp) = 0;
+	virtual bool PeekVideoSampleTimeRange(TRange<FMediaTimeStamp>& TimeRange) { return false; }
 
 	virtual bool DiscardVideoSamples(const TRange<FMediaTimeStamp>& TimeRange, bool bReverse) { return false; }
 	virtual bool DiscardAudioSamples(const TRange<FMediaTimeStamp>& TimeRange, bool bReverse) { return false; }
diff --git a/Engine/Source/Runtime/MediaUtils/Public/MediaSampleQueue.h b/Engine/Source/Runtime/MediaUtils/Public/MediaSampleQueue.h
--- a/Engine/Source/Runtime/MediaUtils/Public/MediaSampleQueue.h
+++ b/Engine/Source/Runtime/MediaUtils/Public/MediaSampleQueue.h
@@ -69,7 +69,7 @@ public:
 		}
 
 		TSharedPtr<SampleType, ESPMode::ThreadSafe> Sample(Samples[0]);
-			
+
 		if (!Sample.IsValid())
 		{
 			return false; // pending flush
@@ -103,6 +103,29 @@ public:
 		return true;
 	}
 
+	virtual bool PeekFrontAndBack(TSharedPtr<SampleType, ESPMode::ThreadSafe>& OutFirstSample, TSharedPtr<SampleType, ESPMode::ThreadSafe>& OutLastSample) override
+	{
+		FScopeLock Lock(&CriticalSection);
+
+		if (Samples.Num() == 0)
+		{
+			return false; // empty queue
+		}
+
+		TSharedPtr<SampleType, ESPMode::ThreadSafe> FirstSample(Samples[0]);
+		TSharedPtr<SampleType, ESPMode::ThreadSafe> LastSample(Samples.Last());
+
+		if (!FirstSample.IsValid() || !LastSample.IsValid())
+		{
+			return false; // pending flush
+		}
+
+		OutFirstSample = FirstSample;
+		OutLastSample = LastSample;
+
+		return true;
+	}
+
 	virtual bool Pop() override
 	{
 		FScopeLock Lock(&CriticalSection);
@@ -150,7 +173,7 @@ public:
 		// - Reverse playback still works with increasing indices in the queue. PTS values will be going down in it, rather than up,
 		//   but the order of indices is still identical.
 		// - The code below must be able to deal with time ranges that span loop points (different secondary sequence indices)
-		
+
 		// Code below assumes a fully specified range, no open bounds!
 		check(TimeRange.HasLowerBound() && TimeRange.HasUpperBound());
 
@@ -215,7 +238,7 @@ public:
 				//
 				// Return the first sample with maximum possible coverage given the time range or nothing if the sample is not yet in the queue
 				// (this yields reproducible results between instances as far as the selection of frames is concerned if the passed in ranges are identical in each run / instance)
-				// 
+				//
 
 				for (int32 Idx = FirstPossibleIndex; Idx <= LastPossibleIndex; ++Idx)
 				{
diff --git a/Engine/Source/Runtime/MediaUtils/Public/MediaSampleSource.h b/Engine/Source/Runtime/MediaUtils/Public/MediaSampleSource.h
--- a/Engine/Source/Runtime/MediaUtils/Public/MediaSampleSource.h
+++ b/Engine/Source/Runtime/MediaUtils/Public/MediaSampleSource.h
@@ -37,6 +37,17 @@ public:
 	 */
 	virtual bool Peek(TSharedPtr<SampleType, ESPMode::ThreadSafe>& OutSample) = 0;
 
+	/**
+	 * Simultaneously peeks at the next (the frontmost) and last samples in the queue without removing them.
+	 * The samples could be identical if there is only one element in the queue.
+	 *
+	 * @param OutFirstSample Will contain the frontmost sample.
+	 * @param OutLastSample Will contain the last sample, which could be identical to the frontmost one.
+	 * @return true if samples are returned, false otherwise.
+	 * @see Peek, Dequeue, Pop
+	 */
+	virtual bool PeekFrontAndBack(TSharedPtr<SampleType, ESPMode::ThreadSafe>& OutFirstSample, TSharedPtr<SampleType, ESPMode::ThreadSafe>& OutLastSample) = 0;
+
 	/**
 	 * Remove the next sample from the queue.
 	 *
diff --git a/Engine/Source/Runtime/MediaUtils/Public/MediaSamples.h b/Engine/Source/Runtime/MediaUtils/Public/MediaSamples.h
--- a/Engine/Source/Runtime/MediaUtils/Public/MediaSamples.h
+++ b/Engine/Source/Runtime/MediaUtils/Public/MediaSamples.h
@@ -180,6 +180,17 @@ public:
 		return true;
 	}
 
+	virtual bool PeekVideoSampleTimeRange(TRange<FMediaTimeStamp>& TimeRange) override
+	{
+		TSharedPtr<IMediaTextureSample, ESPMode::ThreadSafe> First, Last;
+		if (!VideoSampleQueue.PeekFrontAndBack(First, Last))
+		{
+			return false;
+		}
+		TimeRange = TRange<FMediaTimeStamp>(First->GetTime(), Last->GetTime() + Last->GetDuration());
+		return true;
+	}
+
 	virtual bool DiscardVideoSamples(const TRange<FMediaTimeStamp>& TimeRange, bool bReverse) override
 	{
 		return VideoSampleQueue.Discard(TimeRange, bReverse);
