Excluding some V2 timing code from V1 player's code path to avoid crashes with odd-ball PTS values
#rnx

diff --git a/Engine/Source/Runtime/MediaUtils/Private/MediaPlayerFacade.cpp b/Engine/Source/Runtime/MediaUtils/Private/MediaPlayerFacade.cpp
--- a/Engine/Source/Runtime/MediaUtils/Private/MediaPlayerFacade.cpp
+++ b/Engine/Source/Runtime/MediaUtils/Private/MediaPlayerFacade.cpp
@@ -1714,28 +1714,31 @@ void FMediaPlayerFacade::Flush(bool bExcludePlayer, bool bOnSeek)
 		SeekIndex = 0;
 	}
 
-	// Logically we have no old sample anymore if we did seek
-	// (as in: we will start asking for a new one until we get one - even with a rate of zero, if we had a non-zero one ever before)
-	if (bOnSeek)
+	if (Player.IsValid() && Player->GetPlayerFeatureFlag(IMediaPlayer::EFeatureFlag::UsePlaybackTimingV2))
 	{
-		LastVideoSampleProcessedTimeRange = TRange<FMediaTimeStamp>::Empty();
-	}
-	else
-	{
-		if (!bExcludePlayer && !LastVideoSampleProcessedTimeRange.IsEmpty())
+		// Logically we have no old sample anymore if we did seek
+		// (as in: we will start asking for a new one until we get one - even with a rate of zero, if we had a non-zero one ever before)
+		if (bOnSeek)
 		{
-			// Players will reset their sequence index related values, but keep the playback position. Adjust our record accordingly...
-			int32 LoopIdxS = FMediaTimeStamp::GetSecondaryIndex(LastVideoSampleProcessedTimeRange.GetLowerBoundValue().SequenceIndex);
-			int32 LoopIdxE = FMediaTimeStamp::GetSecondaryIndex(LastVideoSampleProcessedTimeRange.GetUpperBoundValue().SequenceIndex);
-			LastVideoSampleProcessedTimeRange.SetLowerBoundValue(FMediaTimeStamp(LastVideoSampleProcessedTimeRange.GetLowerBoundValue().Time, FMediaTimeStamp::MakeSequenceIndex(0, 0)));
-			LastVideoSampleProcessedTimeRange.SetUpperBoundValue(FMediaTimeStamp(LastVideoSampleProcessedTimeRange.GetUpperBoundValue().Time, FMediaTimeStamp::MakeSequenceIndex(0, LoopIdxE - LoopIdxS)));
+			LastVideoSampleProcessedTimeRange = TRange<FMediaTimeStamp>::Empty();
+		}
+		else
+		{
+			if (!bExcludePlayer && !LastVideoSampleProcessedTimeRange.IsEmpty())
+			{
+				// Players will reset their sequence index related values, but keep the playback position. Adjust our record accordingly...
+				int32 LoopIdxS = FMediaTimeStamp::GetSecondaryIndex(LastVideoSampleProcessedTimeRange.GetLowerBoundValue().SequenceIndex);
+				int32 LoopIdxE = FMediaTimeStamp::GetSecondaryIndex(LastVideoSampleProcessedTimeRange.GetUpperBoundValue().SequenceIndex);
+				LastVideoSampleProcessedTimeRange.SetLowerBoundValue(FMediaTimeStamp(LastVideoSampleProcessedTimeRange.GetLowerBoundValue().Time, FMediaTimeStamp::MakeSequenceIndex(0, 0)));
+				LastVideoSampleProcessedTimeRange.SetUpperBoundValue(FMediaTimeStamp(LastVideoSampleProcessedTimeRange.GetUpperBoundValue().Time, FMediaTimeStamp::MakeSequenceIndex(0, LoopIdxE - LoopIdxS)));
+			}
 		}
-	}
 
-	// Invalidate next video time to fetch (none-audio case)
-	NextEstVideoTimeAtFrameStart.Invalidate();
-	// ...and seek target
-	SeekTargetTime.Invalidate();
+		// Invalidate next video time to fetch (none-audio case)
+		NextEstVideoTimeAtFrameStart.Invalidate();
+		// ...and seek target
+		SeekTargetTime.Invalidate();
+	}
 
 	// V1 only
 	NextVideoSampleTime = FTimespan::MinValue();
