Adding PTS resync capability to MFW's PlayerFacade to better enable us to play video-only life streams from webcams (with V2 players)

diff --git a/Engine/Source/Runtime/MediaUtils/Private/MediaPlayerFacade.cpp b/Engine/Source/Runtime/MediaUtils/Private/MediaPlayerFacade.cpp
--- a/Engine/Source/Runtime/MediaUtils/Private/MediaPlayerFacade.cpp
+++ b/Engine/Source/Runtime/MediaUtils/Private/MediaPlayerFacade.cpp
@@ -2599,12 +2599,57 @@ void FMediaPlayerFacade::PreSampleProcessingTimeHandling()
 	if (!SeekTargetTime.IsValid())
 	{
 		// No seek pending & not paused. Can we / Do we need to prime a non-audio clock?
-		if (!bHaveActiveAudio && !BlockOnRange.IsSet() && !NextEstVideoTimeAtFrameStart.IsValid())
+		if (!bHaveActiveAudio && !BlockOnRange.IsSet())
 		{
-			FMediaTimeStamp VideoTimeStamp;
-			if (Player->GetSamples().PeekVideoSampleTime(VideoTimeStamp))
+			// Nothing at all?
+			if (!NextEstVideoTimeAtFrameStart.IsValid())
+			{
+				// Try getting a new sample time to start things up...
+				FMediaTimeStamp VideoTimeStamp;
+				if (Player->GetSamples().PeekVideoSampleTime(VideoTimeStamp))
+				{
+					NextEstVideoTimeAtFrameStart = FMediaTimeStampSample(VideoTimeStamp, FPlatformTime::Seconds());
+				}
+			}
+			else
 			{
-				NextEstVideoTimeAtFrameStart = FMediaTimeStampSample(VideoTimeStamp, FPlatformTime::Seconds());
+				// We have a time. But if we are actively playing forward...
+				if (CurrentRate > 0.0f)
+				{
+					// ...and got some sample waiting for us...
+					FMediaTimeStamp VideoTimeStamp;
+					if (Player->GetSamples().PeekVideoSampleTime(VideoTimeStamp))
+					{
+						// ...we need to see if the player's next sample might be so far in the future that we need to re-calibrate our timing
+						// (this could happen if the stream has a "gap" in PTS values - e.g. after pausing a live feed from a camera)
+						// (^^^ we do not do this on reverse playback as it is unlikely for such streams and might be thinned, hence show gaps under normal conditions)
+						if (VideoTimeStamp.SequenceIndex == NextEstVideoTimeAtFrameStart.TimeStamp.SequenceIndex)
+						{
+							FTimespan Delta = VideoTimeStamp.Time - NextEstVideoTimeAtFrameStart.TimeStamp.Time;
+							if (GetUnpausedRate() < 0.0f)
+							{
+								Delta = -Delta;
+							}
+
+							// Our threshold for re-calibration is twice the length of the last sample we got
+							// (or 100ms if we have nothing)
+							FTimespan DeltaLimit;
+							if (!LastVideoSampleProcessedTimeRange.IsEmpty())
+							{
+								DeltaLimit = LastVideoSampleProcessedTimeRange.Size<FMediaTimeStamp>().Time * 2;
+							}
+							else
+							{
+								DeltaLimit = FTimespan::FromSeconds(0.100);
+							}
+
+							if (Delta >= DeltaLimit)
+							{
+								NextEstVideoTimeAtFrameStart = FMediaTimeStampSample(VideoTimeStamp, FPlatformTime::Seconds());
+							}
+						}
+					}
+				}
 			}
 		}
 	}
