MediaUtils: Added support of playback ranges to the facade

#rnx

diff --git a/Engine/Source/Runtime/MediaUtils/Private/MediaPlayerFacade.cpp b/Engine/Source/Runtime/MediaUtils/Private/MediaPlayerFacade.cpp
--- a/Engine/Source/Runtime/MediaUtils/Private/MediaPlayerFacade.cpp
+++ b/Engine/Source/Runtime/MediaUtils/Private/MediaPlayerFacade.cpp
@@ -35,7 +35,6 @@
 
 #define MEDIAPLAYERFACADE_DISABLE_BLOCKING 0
 #define MEDIAPLAYERFACADE_TRACE_SINKOVERFLOWS 0
-#define MEDIAPLAYERFACADE_DISABLE_PTSCLAMP 1			// enable to disable any clamping of PTS values to the [0..duration[ range (allowing non-zero-PTS-based material with a duration to play)
 
 
 /** Time spent in media player facade closing media. */
@@ -1176,17 +1175,16 @@ bool FMediaPlayerFacade::Seek(const FTimespan& InTime)
 	FTimespan Time;
 	if (IsDurationValidAndFinite(Duration))
 	{
+		const TRange<FTimespan> ActiveRange = GetActivePlaybackRange();
+
 		if (CurrentPlayer->GetControls().IsLooping())
 		{
-			Time = WrappedModulo(InTime, Duration);
+			const FTimespan ActiveRangeDuration = ActiveRange.GetUpperBoundValue() - ActiveRange.GetLowerBoundValue();
+			Time = WrappedModulo(InTime - ActiveRange.GetLowerBoundValue(), ActiveRangeDuration) + ActiveRange.GetLowerBoundValue();
 		}
 		else
 		{
-#if !MEDIAPLAYERFACADE_DISABLE_PTSCLAMP
-			Time = FTimespan(FMath::Clamp(InTime.GetTicks(), (int64)0L, Duration.GetTicks()));
-#else
-			Time = InTime;
-#endif
+			Time = FTimespan(FMath::Clamp(InTime.GetTicks(), ActiveRange.GetLowerBoundValue().GetTicks(), ActiveRange.GetUpperBoundValue().GetTicks()));
 		}
 	}
 	else
@@ -2672,24 +2670,28 @@ void FMediaPlayerFacade::PostSampleProcessingTimeHandling(FTimespan DeltaTime)
 				// note: infinite duration (e.g. live playback - or players not yet supporting sequence indices on loops, when looping is enabled)
 				// -> no need for special handling as FTimespan::MaxValue() is expected to be returned to signify this, which is quite "infinite" in practical terms
 				FTimespan Duration = Player->GetControls().GetDuration();
+				const TRange<FTimespan> ActiveRange = GetActivePlaybackRange();
+				const FTimespan ActiveRangeStart = ActiveRange.GetLowerBoundValue();
+				const FTimespan ActiveRangeEnd = ActiveRange.GetUpperBoundValue();
 
 				if (Player->GetControls().IsLooping())
 				{
 					if (IsDurationValidAndFinite(Duration))
 					{
+						const FTimespan ActiveRangeDuration = ActiveRange.GetUpperBoundValue() - ActiveRange.GetLowerBoundValue();
 						if (Rate >= 0.0f)
 						{
-							while (NextEstVideoTimeAtFrameStart.TimeStamp.Time >= Duration)
+							while(NextEstVideoTimeAtFrameStart.TimeStamp.Time >= ActiveRangeEnd)
 							{
-								NextEstVideoTimeAtFrameStart.TimeStamp.Time -= Duration;
+								NextEstVideoTimeAtFrameStart.TimeStamp.Time -= ActiveRangeDuration;
 								NextEstVideoTimeAtFrameStart.TimeStamp.SequenceIndex = FMediaTimeStamp::AdjustSecondaryIndex(NextEstVideoTimeAtFrameStart.TimeStamp.SequenceIndex, 1);
 							}
 						}
 						else
 						{
-							while (NextEstVideoTimeAtFrameStart.TimeStamp.Time < FTimespan::Zero())
+							while(NextEstVideoTimeAtFrameStart.TimeStamp.Time < ActiveRangeStart)
 							{
-								NextEstVideoTimeAtFrameStart.TimeStamp.Time += Duration;
+								NextEstVideoTimeAtFrameStart.TimeStamp.Time += ActiveRangeDuration;
 								NextEstVideoTimeAtFrameStart.TimeStamp.SequenceIndex = FMediaTimeStamp::AdjustSecondaryIndex(NextEstVideoTimeAtFrameStart.TimeStamp.SequenceIndex, -1);
 							}
 						}
@@ -2701,17 +2703,17 @@ void FMediaPlayerFacade::PostSampleProcessingTimeHandling(FTimespan DeltaTime)
 					{
 						if (IsDurationValidAndFinite(Duration))
 						{
-							if (NextEstVideoTimeAtFrameStart.TimeStamp.Time >= Duration)
+							if (NextEstVideoTimeAtFrameStart.TimeStamp.Time >= ActiveRangeEnd)
 							{
-								NextEstVideoTimeAtFrameStart.TimeStamp.Time = Duration - FTimespan(1);
+								NextEstVideoTimeAtFrameStart.TimeStamp.Time = ActiveRangeEnd - FTimespan(1);
 							}
 						}
 					}
 					else
 					{
-						if (NextEstVideoTimeAtFrameStart.TimeStamp.Time < FTimespan::Zero())
+						if (NextEstVideoTimeAtFrameStart.TimeStamp.Time < ActiveRangeStart)
 						{
-							NextEstVideoTimeAtFrameStart.TimeStamp.Time = FTimespan::Zero();
+							NextEstVideoTimeAtFrameStart.TimeStamp.Time = ActiveRangeStart;
 						}
 					}
 				}
@@ -2721,6 +2723,28 @@ void FMediaPlayerFacade::PostSampleProcessingTimeHandling(FTimespan DeltaTime)
 }
 
 
+TRange<FTimespan> FMediaPlayerFacade::GetActivePlaybackRange() const
+{
+	TRange<FTimespan> Rng(FTimespan::Zero(), FTimespan::Zero());
+	if (Player)
+	{
+		if (SupportsPlaybackTimeRange())
+		{
+			Rng = GetPlaybackTimeRange(EMediaTimeRangeType::Current);
+		}
+		else
+		{
+			FTimespan Duration = Player->GetControls().GetDuration();
+			if (Duration <= FTimespan::Zero())
+			{
+				Duration = FTimespan::MaxValue();
+			}
+			Rng.SetUpperBound(Duration);
+		}
+	}
+	return Rng;
+}
+
 bool FMediaPlayerFacade::GetCurrentPlaybackTimeRange(TRange<FMediaTimeStamp>& TimeRange, float Rate, FTimespan DeltaTime, bool bPurgeSampleRelated) const
 {
 	/*
@@ -2860,16 +2884,18 @@ bool FMediaPlayerFacade::GetCurrentPlaybackTimeRange(TRange<FMediaTimeStamp>& Ti
 	}
 
 	const FTimespan Duration = Player->GetControls().GetDuration();
+	TRange<FTimespan> ActiveRange = GetActivePlaybackRange();
 
 	// We need a valid duration for the next steps (we may not have one e.g. for live material)
 	if (IsDurationValidAndFinite(Duration))
 	{
+		const FTimespan ActiveRangeDuration = ActiveRange.GetUpperBoundValue() - ActiveRange.GetLowerBoundValue();
 		// If we are looping we check to prepare proper ranges should we wrap around either end of the media...
 		// (we do not clamp in the non-looping case as the rest of the code should deal with that fine)
 		if (Player->GetControls().IsLooping())
 		{
-			FTimespan WrappedStart = WrappedModulo(TimeRange.GetLowerBoundValue().Time, Duration);
-			FTimespan WrappedEnd = WrappedModulo(TimeRange.GetUpperBoundValue().Time, Duration);
+			FTimespan WrappedStart = WrappedModulo(TimeRange.GetLowerBoundValue().Time - ActiveRange.GetLowerBoundValue(), ActiveRangeDuration) + ActiveRange.GetLowerBoundValue();
+			FTimespan WrappedEnd = WrappedModulo(TimeRange.GetUpperBoundValue().Time - ActiveRange.GetLowerBoundValue(), ActiveRangeDuration) + ActiveRange.GetLowerBoundValue();
 			if (WrappedStart > WrappedEnd)
 			{
 				if (WrappedStart != TimeRange.GetLowerBoundValue().Time)
@@ -2884,10 +2910,8 @@ bool FMediaPlayerFacade::GetCurrentPlaybackTimeRange(TRange<FMediaTimeStamp>& Ti
 		}
 		else
 		{
-#if !MEDIAPLAYERFACADE_DISABLE_PTSCLAMP
-			TimeRange.SetLowerBoundValue(FMediaTimeStamp(FMath::Clamp(TimeRange.GetLowerBoundValue().Time, FTimespan::Zero(), Duration), TimeRange.GetLowerBoundValue().SequenceIndex));
-			TimeRange.SetUpperBoundValue(FMediaTimeStamp(FMath::Clamp(TimeRange.GetUpperBoundValue().Time, FTimespan::Zero(), Duration), TimeRange.GetUpperBoundValue().SequenceIndex));
-#endif
+			TimeRange.SetLowerBoundValue(FMediaTimeStamp(FMath::Clamp(TimeRange.GetLowerBoundValue().Time, ActiveRange.GetLowerBoundValue(), ActiveRange.GetUpperBoundValue()), TimeRange.GetLowerBoundValue().SequenceIndex));
+			TimeRange.SetUpperBoundValue(FMediaTimeStamp(FMath::Clamp(TimeRange.GetUpperBoundValue().Time, ActiveRange.GetLowerBoundValue(), ActiveRange.GetUpperBoundValue()), TimeRange.GetUpperBoundValue().SequenceIndex));
 		}
 	}
 
@@ -3041,13 +3065,12 @@ bool FMediaPlayerFacade::IsVideoSampleStillGood(const TRange<FMediaTimeStamp>& L
 			// Is looping off?
 			if (!Player->GetControls().IsLooping())
 			{
-#if !MEDIAPLAYERFACADE_DISABLE_PTSCLAMP
 				// Yes. We clamp the range to the duration of the video to avoid looking at non-existent "next" frames... (unless we have no duration)
 				if (IsDurationValidAndFinite(Duration))
 				{
-					TimeRange0 = TRange<FMediaTimeStamp>::Intersection(TimeRange0, TRange<FMediaTimeStamp>(FMediaTimeStamp(FTimespan::Zero(), 0), FMediaTimeStamp(Duration, 0)));
+					const TRange<FTimespan> ActiveRange = GetActivePlaybackRange();
+					TimeRange0 = TRange<FMediaTimeStamp>::Intersection(TimeRange0, TRange<FMediaTimeStamp>(FMediaTimeStamp(ActiveRange.GetLowerBoundValue(), 0), FMediaTimeStamp(ActiveRange.GetUpperBoundValue(), 0)));
 				}
-#endif
 			}
 		}
 
diff --git a/Engine/Source/Runtime/MediaUtils/Public/MediaPlayerFacade.h b/Engine/Source/Runtime/MediaUtils/Public/MediaPlayerFacade.h
--- a/Engine/Source/Runtime/MediaUtils/Public/MediaPlayerFacade.h
+++ b/Engine/Source/Runtime/MediaUtils/Public/MediaPlayerFacade.h
@@ -938,6 +938,7 @@ private:
 private:
 	MEDIAUTILS_API TRange<FMediaTimeStamp> GetAdjustedBlockOnRange() const;
 
+	TRange<FTimespan> GetActivePlaybackRange() const;
 	void ReInit();
 
 	class FBlockOnRange
