MediaUtils: Fixed wrong time threshold with dropping old subtitles

#rnx

diff --git a/Engine/Source/Runtime/MediaUtils/Private/MediaPlayerFacade.cpp b/Engine/Source/Runtime/MediaUtils/Private/MediaPlayerFacade.cpp
--- a/Engine/Source/Runtime/MediaUtils/Private/MediaPlayerFacade.cpp
+++ b/Engine/Source/Runtime/MediaUtils/Private/MediaPlayerFacade.cpp
@@ -85,7 +85,7 @@ DECLARE_DWORD_ACCUMULATOR_STAT(TEXT("MediaPlayerFacade TotalPurgedCaptionSamples
 static const double kMaxTimeSinceFrameStart = 0.300;			// max seconds we allow between the start of the frame and the player facade timing computations (to catch suspended apps & debugging)
 static const double kMaxTimeSinceAudioTimeSampling = 0.250;		// max seconds we allow to have passed between the last audio timing sampling and the player facade timing computations (to catch suspended apps & debugging - some platforms do update audio at a farily low rate: hence the big tollerance)
 static const double kOutdatedVideoSamplesTolerance = 0.080;		// seconds video samples are allowed to be "too old" to stay in the player's output queue despite of calculations indicating they need to go
-static const double kOutdatedSubtitleSamplesTolerance = 0.050;	// seconds subtitle samples are allowed to be "too old" to stay in the player's output queue despite of calculations indicating they need to go
+static const double kOutdatedSubtitleSamplesTolerance = 1.0;	// seconds subtitle samples are allowed to be "too old" to stay in the player's output queue despite of calculations indicating they need to go
 static const double kOutdatedSamplePurgeRange = 1.0;			// milliseconds for pseudo DT timespan used with async purging of outdated video samples
 static const int32	kMinFramesInVideoQueueToPurge = 3;			// we only consider purging any old frames from the video queue if more than these are present (to not kill a slow playback entirely)
 static const int32	kMinFramesInSubtitleQueueToPurge = 3;		// we only consider purging any old frames from the subtitle queue if more than these are present (to not kill a slow playback entirely)
@@ -3427,7 +3427,7 @@ void FMediaPlayerFacade::ProcessSubtitleSamplesV1(IMediaSamples& Samples, TRange
 		if (Sample.IsValid() && !SubtitleSampleSinks.Enqueue(Sample.ToSharedRef()))
 		{
 			FString Caption = Sample->GetText().ToString();
-			UE_LOG(LogMediaUtils, Log, TEXT("New caption @%.3f: %s"), Sample->GetTime().Time.GetTotalSeconds(), *Caption);
+			//UE_LOG(LogMediaUtils, Log, TEXT("New caption @%.3f: %s"), Sample->GetTime().Time.GetTotalSeconds(), *Caption);
 
 #if MEDIAPLAYERFACADE_TRACE_SINKOVERFLOWS
 			UE_LOG(LogMediaUtils, VeryVerbose, TEXT("PlayerFacade %p: Subtitle sample sink overflow"), this);
@@ -3523,7 +3523,7 @@ void FMediaPlayerFacade::ReceiveMediaEvent(EMediaEvent Event)
 			NumPurged = 0;
 			if (CurrentPlayer->GetSamples().NumVideoSamples() >= kMinFramesInCaptionQueueToPurge)
 			{
-				NumPurged = CurrentPlayer->GetSamples().PurgeOutdatedCaptionSamples(TimeRange.GetLowerBoundValue(), bReverse, FTimespan::FromSeconds(kOutdatedVideoSamplesTolerance * RateFactor));
+				NumPurged = CurrentPlayer->GetSamples().PurgeOutdatedCaptionSamples(TimeRange.GetLowerBoundValue(), bReverse, FTimespan::FromSeconds(kOutdatedSubtitleSamplesTolerance * RateFactor));
 			}
 			SET_DWORD_STAT(STAT_MediaUtils_FacadeNumPurgedSubtitleSamples, NumPurged);
 			INC_DWORD_STAT_BY(STAT_MediaUtils_FacadeTotalPurgedSubtitleSamples, NumPurged);
@@ -3532,7 +3532,7 @@ void FMediaPlayerFacade::ReceiveMediaEvent(EMediaEvent Event)
 			NumPurged = 0;
 			if (CurrentPlayer->GetSamples().NumVideoSamples() >= kMinFramesInSubtitleQueueToPurge)
 			{
-				NumPurged = CurrentPlayer->GetSamples().PurgeOutdatedSubtitleSamples(TimeRange.GetLowerBoundValue(), bReverse, FTimespan::FromSeconds(kOutdatedVideoSamplesTolerance * RateFactor));
+				NumPurged = CurrentPlayer->GetSamples().PurgeOutdatedSubtitleSamples(TimeRange.GetLowerBoundValue(), bReverse, FTimespan::FromSeconds(kOutdatedSubtitleSamplesTolerance * RateFactor));
 			}
 			SET_DWORD_STAT(STAT_MediaUtils_FacadeNumPurgedCaptionSamples, NumPurged);
 			INC_DWORD_STAT_BY(STAT_MediaUtils_FacadeTotalPurgedCaptionSamples, NumPurged);
