MediaServices: Fix crash when a bad downsampling occurs.

diff --git a/Engine/Source/Runtime/MediaUtils/Private/MediaAudioResampler.cpp b/Engine/Source/Runtime/MediaUtils/Private/MediaAudioResampler.cpp
--- a/Engine/Source/Runtime/MediaUtils/Private/MediaAudioResampler.cpp
+++ b/Engine/Source/Runtime/MediaUtils/Private/MediaAudioResampler.cpp
@@ -378,14 +378,11 @@ uint32 FMediaAudioResampler::Generate(float* Output, FMediaTimeStamp& OutTime, c
 			UE_LOG(LogMediaUtils, VeryVerbose, TEXT("[AUDIO DEQUEUE] PlayerTime: %s, SampleTime: %s, Contiguous: %d, Frames: %d"), *Time.ToString(), *InputTime.Time.ToString(), bIsContiguous, NextSample->GetFrames());
 		}
 
-		check(Input.Num() > 0);
-
 		// skip input if too small
-		if ((FrameIndex >= 0) && (FrameIndex >= InputFrames))
+		if (FrameIndex >= 0 && FrameIndex >= InputFrames)
 		{
 			LastFrameIndex = MIN_int32;
 			FrameIndex -= InputFrames;
-
 			continue;
 		}
 
@@ -407,24 +404,36 @@ uint32 FMediaAudioResampler::Generate(float* Output, FMediaTimeStamp& OutTime, c
 		// get current & next input frame
 		if (FrameIndex != LastFrameIndex)
 		{
-			const float* InputFrame = Input.GetData();
+			const float* InputFrameBuffer = Input.GetData();
 
 			if (FrameIndex == INDEX_NONE)
 			{
+				if (!ensure((uint32)(Input.Num()) >= OutputChannels))
+				{
+					ClearInput();
+					break;
+				}
+
 				// we're still in the last frame of the previous input buffer
 				for (uint32 Channel = 0; Channel < OutputChannels; ++Channel)
 				{
-					NextFrame[Channel] = InputFrame[Channel];
+					NextFrame[Channel] = InputFrameBuffer[Channel];
 				}
 			}
 			else if (FrameIndex == InputFrames - 1)
 			{
+				if (!ensure(Input.Num() >= ((FrameIndex * OutputChannels) + OutputChannels)))
+				{
+					ClearInput();
+					break;
+				}
+
 				// reached the end of the input buffer; cache last frame
-				InputFrame += FrameIndex * OutputChannels;
+				InputFrameBuffer += FrameIndex * OutputChannels;
 
 				for (uint32 Channel = 0; Channel < OutputChannels; ++Channel)
 				{
-					CurrentFrame[Channel] = InputFrame[Channel];
+					CurrentFrame[Channel] = InputFrameBuffer[Channel];
 				}
 
 				LastFrameIndex = MIN_int32;
@@ -434,19 +443,25 @@ uint32 FMediaAudioResampler::Generate(float* Output, FMediaTimeStamp& OutTime, c
 			}
 			else
 			{
+				if (!ensure(Input.Num() >= ((FrameIndex * OutputChannels) + OutputChannels + OutputChannels)))
+				{
+					ClearInput();
+					break;
+				}
+
 				// we're in the current input buffer
-				InputFrame += FrameIndex * OutputChannels;
+				InputFrameBuffer += FrameIndex * OutputChannels;
 
 				for (uint32 Channel = 0; Channel < OutputChannels; ++Channel)
 				{
-					CurrentFrame[Channel] = InputFrame[Channel];
+					CurrentFrame[Channel] = InputFrameBuffer[Channel];
 				}
 
-				InputFrame += OutputChannels;
+				InputFrameBuffer += OutputChannels;
 
 				for (uint32 Channel = 0; Channel < OutputChannels; ++Channel)
 				{
-					NextFrame[Channel] = InputFrame[Channel];
+					NextFrame[Channel] = InputFrameBuffer[Channel];
 				}
 			}
 
