Media: Resync facade to video timestamp when there is no audio and the estimated video time gets ahead of the decoded video time.

#rnx

diff --git a/Engine/Plugins/Media/WmfMedia/Source/WmfMedia/Private/Player/WmfMediaTracks.cpp b/Engine/Plugins/Media/WmfMedia/Source/WmfMedia/Private/Player/WmfMediaTracks.cpp
--- a/Engine/Plugins/Media/WmfMedia/Source/WmfMedia/Private/Player/WmfMediaTracks.cpp
+++ b/Engine/Plugins/Media/WmfMedia/Source/WmfMedia/Private/Player/WmfMediaTracks.cpp
@@ -595,7 +595,9 @@ void FWmfMediaTracks::FlushSamples()
 
 IMediaSamples::EFetchBestSampleResult FWmfMediaTracks::FetchBestVideoSampleForTimeRange(const TRange<FMediaTimeStamp>& TimeRange, TSharedPtr<IMediaTextureSample, ESPMode::ThreadSafe>& OutSample, bool bReverse, bool bConsistentResult)
 {
-	EFetchBestSampleResult Ret = VideoSampleQueue.FetchBestSampleForTimeRange(TimeRange, OutSample, bReverse, bConsistentResult) ? EFetchBestSampleResult::Ok : EFetchBestSampleResult::NoSample;
+	EMediaSampleQueueFetchResult FetchRes =	VideoSampleQueue.FetchBestSampleForTimeRange(TimeRange, OutSample, bReverse, bConsistentResult);
+	EFetchBestSampleResult Ret = FetchRes == EMediaSampleQueueFetchResult::Found ? EFetchBestSampleResult::Ok :
+								 FetchRes == EMediaSampleQueueFetchResult::None ? EFetchBestSampleResult::NoSample : EFetchBestSampleResult::PurgedToEmpty;
 
 	Session->RequestMoreVideoData();
 
diff --git a/Engine/Source/Runtime/Media/Public/IMediaSamples.h b/Engine/Source/Runtime/Media/Public/IMediaSamples.h
--- a/Engine/Source/Runtime/Media/Public/IMediaSamples.h
+++ b/Engine/Source/Runtime/Media/Public/IMediaSamples.h
@@ -120,6 +120,7 @@ public:
 	{
 		Ok = 0,
 		NoSample,
+		PurgedToEmpty,
 		NotSupported,
 	};
 	virtual EFetchBestSampleResult FetchBestVideoSampleForTimeRange(const TRange<FMediaTimeStamp>& TimeRange, TSharedPtr<IMediaTextureSample, ESPMode::ThreadSafe>& OutSample, bool bReverse, bool bConsistentResult)
diff --git a/Engine/Source/Runtime/MediaUtils/Private/MediaPlayerFacade.cpp b/Engine/Source/Runtime/MediaUtils/Private/MediaPlayerFacade.cpp
--- a/Engine/Source/Runtime/MediaUtils/Private/MediaPlayerFacade.cpp
+++ b/Engine/Source/Runtime/MediaUtils/Private/MediaPlayerFacade.cpp
@@ -1729,7 +1729,7 @@ void FMediaPlayerFacade::Flush(bool bExcludePlayer, bool bOnSeek)
 			}
 		}
 
-		// Invalidate next video time to fetch (none-audio case)
+		// Invalidate next video time to fetch (non-audio case)
 		NextEstVideoTimeAtFrameStart.Invalidate();
 		// ...and seek target
 		SeekTargetTime.Invalidate();
@@ -3168,41 +3168,55 @@ bool FMediaPlayerFacade::ProcessVideoSamples(IMediaSamples& Samples, const TRang
 			return true;
 		}
 
-		switch (Samples.FetchBestVideoSampleForTimeRange(TimeRange, Sample, bReverse, BlockOnRange.IsSet()))
+		switch(Samples.FetchBestVideoSampleForTimeRange(TimeRange, Sample, bReverse, BlockOnRange.IsSet()))
 		{
-		case IMediaSamples::EFetchBestSampleResult::Ok:
-			break;
-
-		case IMediaSamples::EFetchBestSampleResult::NoSample:
-		{
-			break;
-		}
+			case IMediaSamples::EFetchBestSampleResult::Ok:
+			{
+				break;
+			}
 
-		case IMediaSamples::EFetchBestSampleResult::NotSupported:
-		{
-			//
-			// Fallback for players supporting V2 timing, but do not supply FetchBestVideoSampleForTimeRange() due to some
-			// custom implementation of IMediaSamples (here to ease adoption of the new timing code - eventually should go away)
-			//
+			case IMediaSamples::EFetchBestSampleResult::NoSample:
+			{
+				break;
+			}
 
-			// Find newest sample that satisfies the time range
-			// (the FetchXYZ() code does not work well with a lower range limit at all - we ask for a "up to" type range instead
-			//  and limit the other side of the range in code here to not change the older logic & possibly cause trouble in old code)
-			TRange<FMediaTimeStamp> TempRange = bReverse ? TRange<FMediaTimeStamp>::AtLeast(TimeRange.GetUpperBoundValue()) : TRange<FMediaTimeStamp>::AtMost(TimeRange.GetUpperBoundValue());
-			while (Samples.FetchVideo(TempRange, Sample))
-				;
-			if (Sample.IsValid() &&
-				((!bReverse && ((Sample->GetTime() + Sample->GetDuration()) > TimeRange.GetLowerBoundValue())) ||
-					(bReverse && ((Sample->GetTime() - Sample->GetDuration()) < TimeRange.GetLowerBoundValue()))))
+			case IMediaSamples::EFetchBestSampleResult::PurgedToEmpty:
 			{
-				// Sample is good - nothing more to do here
+				// When there is no audio to sync to then we are extrapolating the next expected video timestamp
+				// from the last plus the elapsed deltatime, which may overshoot the next decoder output.
+				// In this case we resynchronize the timestamp to the next available video frame.
+				if (!HaveAudioPlayback())
+				{
+					NextEstVideoTimeAtFrameStart.Invalidate();
+				}
+				break;
 			}
-			else
+
+			case IMediaSamples::EFetchBestSampleResult::NotSupported:
 			{
-				Sample.Reset();
+				//
+				// Fallback for players supporting V2 timing, but do not supply FetchBestVideoSampleForTimeRange() due to some
+				// custom implementation of IMediaSamples (here to ease adoption of the new timing code - eventually should go away)
+				//
+
+				// Find newest sample that satisfies the time range
+				// (the FetchXYZ() code does not work well with a lower range limit at all - we ask for a "up to" type range instead
+				//  and limit the other side of the range in code here to not change the older logic & possibly cause trouble in old code)
+				TRange<FMediaTimeStamp> TempRange = bReverse ? TRange<FMediaTimeStamp>::AtLeast(TimeRange.GetUpperBoundValue()) : TRange<FMediaTimeStamp>::AtMost(TimeRange.GetUpperBoundValue());
+				while (Samples.FetchVideo(TempRange, Sample))
+				{ }
+				if (Sample.IsValid() &&
+					((!bReverse && ((Sample->GetTime() + Sample->GetDuration()) > TimeRange.GetLowerBoundValue())) ||
+						(bReverse && ((Sample->GetTime() - Sample->GetDuration()) < TimeRange.GetLowerBoundValue()))))
+				{
+					// Sample is good - nothing more to do here
+				}
+				else
+				{
+					Sample.Reset();
+				}
+				break;
 			}
-			break;
-		}
 		}
 	}
 	else
@@ -3212,7 +3226,7 @@ bool FMediaPlayerFacade::ProcessVideoSamples(IMediaSamples& Samples, const TRang
 		//
 		TRange<FMediaTimeStamp> TempRange; // fully open range
 		while (Samples.FetchVideo(TempRange, Sample))
-			;
+		{ }
 	}
 
 	// Any sample?
diff --git a/Engine/Source/Runtime/MediaUtils/Private/MediaSamples.cpp b/Engine/Source/Runtime/MediaUtils/Private/MediaSamples.cpp
--- a/Engine/Source/Runtime/MediaUtils/Private/MediaSamples.cpp
+++ b/Engine/Source/Runtime/MediaUtils/Private/MediaSamples.cpp
@@ -141,11 +141,10 @@ void FMediaSamples::FlushSamples()
  */
 FMediaSamples::EFetchBestSampleResult FMediaSamples::FetchBestVideoSampleForTimeRange(const TRange<FMediaTimeStamp> & TimeRange, TSharedPtr<IMediaTextureSample, ESPMode::ThreadSafe>& OutSample, bool bReverse, bool bConsistentResult)
 {
-	if (!VideoSampleQueue.FetchBestSampleForTimeRange(TimeRange, OutSample, bReverse, bConsistentResult))
-	{
-		return EFetchBestSampleResult::NoSample;
-	}
-	return EFetchBestSampleResult::Ok;
+	EMediaSampleQueueFetchResult FetchRes =	VideoSampleQueue.FetchBestSampleForTimeRange(TimeRange, OutSample, bReverse, bConsistentResult);
+	EFetchBestSampleResult Ret = FetchRes == EMediaSampleQueueFetchResult::Found ? EFetchBestSampleResult::Ok :
+								 FetchRes == EMediaSampleQueueFetchResult::None ? EFetchBestSampleResult::NoSample : EFetchBestSampleResult::PurgedToEmpty;
+	return Ret;
 }
 
 /**
diff --git a/Engine/Source/Runtime/MediaUtils/Public/MediaSampleQueue.h b/Engine/Source/Runtime/MediaUtils/Public/MediaSampleQueue.h
--- a/Engine/Source/Runtime/MediaUtils/Public/MediaSampleQueue.h
+++ b/Engine/Source/Runtime/MediaUtils/Public/MediaSampleQueue.h
@@ -21,6 +21,14 @@
 #include "IMediaBinarySample.h"
 #include "IMediaOverlaySample.h"
 
+
+enum class EMediaSampleQueueFetchResult
+{
+	Found,
+	None,
+	PurgedToEmpty
+};
+
 /**
  * Template for media sample queues.
  */
@@ -179,7 +187,7 @@ public:
 		return false;
 	}
 
-	bool FetchBestSampleForTimeRange(const TRange<FMediaTimeStamp>& TimeRange, TSharedPtr<SampleType, ESPMode::ThreadSafe>& OutSample, bool bReverse, bool bConsistentResult)
+	EMediaSampleQueueFetchResult FetchBestSampleForTimeRange(const TRange<FMediaTimeStamp>& TimeRange, TSharedPtr<SampleType, ESPMode::ThreadSafe>& OutSample, bool bReverse, bool bConsistentResult)
 	{
 		// Notes:
 		// - Reverse playback still works with increasing indices in the queue. PTS values will be going down in it, rather than up,
@@ -284,7 +292,7 @@ public:
 		}
 
 		// Return true if we got a sample...
-		return (OutSample.IsValid());
+		return OutSample.IsValid() ? EMediaSampleQueueFetchResult::Found : (NumOldSamplesAtBegin && Samples.IsEmpty() ? EMediaSampleQueueFetchResult::PurgedToEmpty : EMediaSampleQueueFetchResult::None);
 	}
 
 
