MediaFramework: GC the MediaOption. IMediaOption should be used as a temporary argument, but new players keep it alive for more than one frame because of async calls. The MediaOption needs to be GCed correctly.

diff --git a/Engine/Plugins/Experimental/TextureMediaPlayer/Source/TextureMediaPlayer/Private/TextureMediaPlayer.cpp b/Engine/Plugins/Experimental/TextureMediaPlayer/Source/TextureMediaPlayer/Private/TextureMediaPlayer.cpp
--- a/Engine/Plugins/Experimental/TextureMediaPlayer/Source/TextureMediaPlayer/Private/TextureMediaPlayer.cpp
+++ b/Engine/Plugins/Experimental/TextureMediaPlayer/Source/TextureMediaPlayer/Private/TextureMediaPlayer.cpp
@@ -18,7 +18,6 @@ DEFINE_LOG_CATEGORY(LogTextureMediaPlayer);
  */
 FTextureMediaPlayer::FTextureMediaPlayer(IMediaEventSink& InEventSink)
 	: EventSink(InEventSink)
-	, OptionInterface(nullptr)
 	, State(EMediaState::Closed)
 	, Status(EMediaStatus::None)
 	, bWasClosedOnError(false)
@@ -99,9 +98,6 @@ bool FTextureMediaPlayer::Open(const FString& _InUrl, const IMediaOptions* Optio
 	// Clear out our work variables
 	bWasClosedOnError = false;
 
-	// Remember the option interface to poll for changes during playback.
-	OptionInterface = Options;
-
 	// Get a writable copy of the URL so we can override it for debugging.
 	MediaUrl = InUrl;
 
diff --git a/Engine/Plugins/Experimental/TextureMediaPlayer/Source/TextureMediaPlayer/Private/TextureMediaPlayer.h b/Engine/Plugins/Experimental/TextureMediaPlayer/Source/TextureMediaPlayer/Private/TextureMediaPlayer.h
--- a/Engine/Plugins/Experimental/TextureMediaPlayer/Source/TextureMediaPlayer/Private/TextureMediaPlayer.h
+++ b/Engine/Plugins/Experimental/TextureMediaPlayer/Source/TextureMediaPlayer/Private/TextureMediaPlayer.h
@@ -153,9 +153,6 @@ private:
 	/** Media player Guid */
 	FGuid											PlayerGuid;
 
-	/** Option interface **/
-	const IMediaOptions*							OptionInterface;
-
 	/**  */
 	TAtomic<EMediaState>							State;
 	TAtomic<EMediaStatus>							Status;
diff --git a/Engine/Plugins/Media/ElectraPlayer/Source/ElectraPlayerPlugin/Public/IElectraPlayerPluginModule.h b/Engine/Plugins/Media/ElectraPlayer/Source/ElectraPlayerPlugin/Public/IElectraPlayerPluginModule.h
--- a/Engine/Plugins/Media/ElectraPlayer/Source/ElectraPlayerPlugin/Public/IElectraPlayerPluginModule.h
+++ b/Engine/Plugins/Media/ElectraPlayer/Source/ElectraPlayerPlugin/Public/IElectraPlayerPluginModule.h
@@ -8,11 +8,14 @@
 #include "Misc/Timespan.h"
 #include "Containers/Array.h"
 #include "IMediaOptions.h"
+#include "UObject/Object.h"
+#include "UObject/StrongObjectPtrTemplates.h"
 
 class IAnalyticsProviderET;
 class IMediaEventSink;
 class IMediaPlayer;
 class IElectraPlayerDataCache;
+class UObject;
 
 /**
  * This class is used to get safe access to an IMediaOptions interface.
@@ -68,6 +71,7 @@ class FElectraSafeMediaOptionInterface : public IElectraSafeMediaOptionInterface
 public:
 	FElectraSafeMediaOptionInterface(IMediaOptions* InOwner)
 		: Owner(InOwner)
+		, OwnerObject(InOwner ? InOwner->ToUObject() : nullptr)
 	{ }
 	virtual ~FElectraSafeMediaOptionInterface()
 	{
@@ -77,6 +81,7 @@ public:
 	{
 		FScopeLock lock(&OwnerLock);
 		Owner = nullptr;
+		OwnerObject = nullptr;
 	}
 	virtual void Lock() override
 	{
@@ -93,6 +98,7 @@ public:
 private:
 	FCriticalSection OwnerLock;
 	IMediaOptions* Owner = nullptr;
+	TStrongObjectPtr<const UObject> OwnerObject;
 };
 
 
diff --git a/Engine/Source/Runtime/Media/Public/IMediaOptions.h b/Engine/Source/Runtime/Media/Public/IMediaOptions.h
--- a/Engine/Source/Runtime/Media/Public/IMediaOptions.h
+++ b/Engine/Source/Runtime/Media/Public/IMediaOptions.h
@@ -93,4 +93,10 @@ public:
 	 * @return true if the option is set, false otherwise.
 	 */
 	virtual bool HasMediaOption(const FName& Key) const = 0;
+
+	/**
+	 * @return a object that can be used by the UObject garbage collection system to keep this object alive.
+	 * Players should not store a pointer of the interface but some players use the interface in async functionality and the object needs to be kept alive during those.
+	 */
+	virtual const UObject* ToUObject() const = 0;
 };
diff --git a/Engine/Source/Runtime/MediaAssets/Private/Assets/MediaSource.cpp b/Engine/Source/Runtime/MediaAssets/Private/Assets/MediaSource.cpp
--- a/Engine/Source/Runtime/MediaAssets/Private/Assets/MediaSource.cpp
+++ b/Engine/Source/Runtime/MediaAssets/Private/Assets/MediaSource.cpp
@@ -193,6 +193,11 @@ bool UMediaSource::HasMediaOption(const FName& Key) const
 	return MediaOptionsMap.Contains(Key);
 }
 
+const UObject* UMediaSource::ToUObject() const
+{
+	return this;
+}
+
 
 void UMediaSource::SetMediaOptionBool(const FName& Key, bool Value)
 {
diff --git a/Engine/Source/Runtime/MediaAssets/Public/MediaSource.h b/Engine/Source/Runtime/MediaAssets/Public/MediaSource.h
--- a/Engine/Source/Runtime/MediaAssets/Public/MediaSource.h
+++ b/Engine/Source/Runtime/MediaAssets/Public/MediaSource.h
@@ -123,6 +123,7 @@ public:
 	MEDIAASSETS_API virtual FText GetMediaOption(const FName& Key, const FText& DefaultValue) const override;
 	MEDIAASSETS_API virtual TSharedPtr<FDataContainer, ESPMode::ThreadSafe> GetMediaOption(const FName& Key, const TSharedPtr<FDataContainer, ESPMode::ThreadSafe>& DefaultValue) const override;
 	MEDIAASSETS_API virtual bool HasMediaOption(const FName& Key) const override;
+	MEDIAASSETS_API virtual const UObject* ToUObject() const override;
 
 	/** Set a boolean parameter to pass to the player. */
 	UFUNCTION(BlueprintCallable, meta=(DisplayName = "SetMediaOption (boolean)"), Category = "Media|MediaSource")
diff --git a/Engine/Source/Runtime/MediaUtils/Private/MediaPlayerFacade.cpp b/Engine/Source/Runtime/MediaUtils/Private/MediaPlayerFacade.cpp
--- a/Engine/Source/Runtime/MediaUtils/Private/MediaPlayerFacade.cpp
+++ b/Engine/Source/Runtime/MediaUtils/Private/MediaPlayerFacade.cpp
@@ -747,7 +747,13 @@ class FMediaPlayerLifecycleManagerDelegateOpenRequest : public IMediaPlayerLifec
 {
 public:
 	FMediaPlayerLifecycleManagerDelegateOpenRequest(const FString& InUrl, const IMediaOptions* InOptions, const FMediaPlayerOptions* InPlayerOptions, IMediaPlayerFactory* InPlayerFactory, TSharedPtr<IMediaPlayer, ESPMode::ThreadSafe> InReusedPlayer, bool bInWillCreatePlayer, uint32 InWillUseNewResources)
-		: Url(InUrl), Options(InOptions), PlayerFactory(InPlayerFactory), ReusedPlayer(InReusedPlayer), bWillCreatePlayer(bInWillCreatePlayer), NewResources(InWillUseNewResources)
+		: Url(InUrl)
+		, Options(InOptions)
+		, OptionsObject(InOptions ? InOptions->ToUObject() : nullptr)
+		, PlayerFactory(InPlayerFactory)
+		, ReusedPlayer(InReusedPlayer)
+		, bWillCreatePlayer(bInWillCreatePlayer)
+		, NewResources(InWillUseNewResources)
 	{
 		if (InPlayerOptions)
 		{
@@ -785,8 +791,15 @@ public:
 		return !!(NewResources & ResourceFlags);
 	}
 
+	const TSharedPtr<IMediaPlayer, ESPMode::ThreadSafe>& GetReusedPlayer() const
+	{
+		return ReusedPlayer;
+	}
+
+private:
 	FString Url;
 	const IMediaOptions* Options;
+	TStrongObjectPtr<const UObject> OptionsObject;
 	TOptional<FMediaPlayerOptions> PlayerOptions;
 	IMediaPlayerFactory* PlayerFactory;
 	TSharedPtr<IMediaPlayer, ESPMode::ThreadSafe> ReusedPlayer;
@@ -815,7 +828,7 @@ public:
 		if (TSharedPtr<FMediaPlayerFacade, ESPMode::ThreadSafe> PinnedFacade = Facade.Pin())
 		{
 			const FMediaPlayerLifecycleManagerDelegateOpenRequest* OR = static_cast<const FMediaPlayerLifecycleManagerDelegateOpenRequest*>(OpenRequest.Get());
-			if (PinnedFacade->ContinueOpen(AsShared(), OR->Url, OR->Options, OR->PlayerOptions.IsSet() ? &OR->PlayerOptions.GetValue() : nullptr, OR->PlayerFactory, OR->ReusedPlayer, OR->bWillCreatePlayer, InstanceID))
+			if (PinnedFacade->ContinueOpen(AsShared(), OR->GetUrl(), OR->GetOptions(), OR->GetPlayerOptions(), OR->GetPlayerFactory(), OR->GetReusedPlayer(), OR->WillCreateNewPlayer(), InstanceID))
 			{
 				SubmittedRequest = true;
 			}
