MediaServices: Fix GC issue introduced by 35867793. There was a circular reference inside the EpicMediaBasePlayer. A TSharedPtr mixed with GC.

#rnx

diff --git a/Engine/Source/Runtime/MediaUtils/Private/MediaPlayerFacade.cpp b/Engine/Source/Runtime/MediaUtils/Private/MediaPlayerFacade.cpp
--- a/Engine/Source/Runtime/MediaUtils/Private/MediaPlayerFacade.cpp
+++ b/Engine/Source/Runtime/MediaUtils/Private/MediaPlayerFacade.cpp
@@ -768,7 +768,7 @@ public:
 
 	virtual const IMediaOptions* GetOptions() const override
 	{
-		return Options;
+		return OptionsObject.IsValid() ? Options : nullptr;
 	}
 
 	virtual const FMediaPlayerOptions* GetPlayerOptions() const override
@@ -799,7 +799,7 @@ public:
 private:
 	FString Url;
 	const IMediaOptions* Options;
-	TStrongObjectPtr<const UObject> OptionsObject;
+	FWeakObjectPtr OptionsObject;
 	TOptional<FMediaPlayerOptions> PlayerOptions;
 	IMediaPlayerFactory* PlayerFactory;
 	TSharedPtr<IMediaPlayer, ESPMode::ThreadSafe> ReusedPlayer;
