MediaUtils: Fixed a flush issue when changing playback direction; fixed missing initial track selection when receiving track change event

#rnx

diff --git a/Engine/Source/Runtime/MediaUtils/Private/MediaPlayerFacade.cpp b/Engine/Source/Runtime/MediaUtils/Private/MediaPlayerFacade.cpp
--- a/Engine/Source/Runtime/MediaUtils/Private/MediaPlayerFacade.cpp
+++ b/Engine/Source/Runtime/MediaUtils/Private/MediaPlayerFacade.cpp
@@ -1518,7 +1518,8 @@ bool FMediaPlayerFacade::SetRate(float Rate)
 
 	if ((LastRate * Rate) < 0.0f)
 	{
-		Flush(); // direction change
+		// direction change
+		Flush(Player->GetPlayerFeatureFlag(IMediaPlayer::EFeatureFlag::UsePlaybackTimingV2), false);
 	}
 	else
 	{
@@ -1939,11 +1940,15 @@ void FMediaPlayerFacade::ProcessEvent(EMediaEvent Event, bool bIsBroadcastAllowe
 	else if (Event == EMediaEvent::TracksChanged)
 	{
 		SelectDefaultTracks();
-
-		if (Player.IsValid() && !Player->GetPlayerFeatureFlag(IMediaPlayer::EFeatureFlag::UsePlaybackTimingV2))
+		if (Player.IsValid())
 		{
-			// Execute flush for older players only
-			Flush();
+			// Apply track selection immediately so the selection can be queried.
+			UpdateTrackSelectionWithPlayer();
+			if (!Player->GetPlayerFeatureFlag(IMediaPlayer::EFeatureFlag::UsePlaybackTimingV2))
+			{
+				// Execute flush for older players only
+				Flush();
+			}
 		}
 	}
 	else if (Event == EMediaEvent::SeekCompleted)
