MediaUtils: Removed 'this' pointer from facade log messages

#rnx

diff --git a/Engine/Source/Runtime/MediaUtils/Private/MediaPlayerFacade.cpp b/Engine/Source/Runtime/MediaUtils/Private/MediaPlayerFacade.cpp
--- a/Engine/Source/Runtime/MediaUtils/Private/MediaPlayerFacade.cpp
+++ b/Engine/Source/Runtime/MediaUtils/Private/MediaPlayerFacade.cpp
@@ -1694,7 +1694,7 @@ bool FMediaPlayerFacade::BlockOnFetch() const
 
 void FMediaPlayerFacade::Flush(bool bExcludePlayer, bool bOnSeek)
 {
-	UE_LOG(LogMediaUtils, Verbose, TEXT("PlayerFacade %p: Flushing sinks"), this);
+	UE_LOG(LogMediaUtils, Verbose, TEXT("PlayerFacade: Flushing sinks"));
 
 	FScopeLock Lock(&CriticalSection);
 
@@ -1930,11 +1930,11 @@ void FMediaPlayerFacade::ProcessEvent(EMediaEvent Event, bool bIsBroadcastAllowe
 
 		if (MediaInfo.IsEmpty())
 		{
-			UE_LOG(LogMediaUtils, Verbose, TEXT("PlayerFacade %p: Media Info: n/a"), this);
+			UE_LOG(LogMediaUtils, Verbose, TEXT("PlayerFacade: Media Info: n/a"));
 		}
 		else
 		{
-			UE_LOG(LogMediaUtils, Verbose, TEXT("PlayerFacade %p: Media Info:\n%s"), this, *MediaInfo);
+			UE_LOG(LogMediaUtils, Verbose, TEXT("PlayerFacade: Media Info:\n%s"), *MediaInfo);
 		}
 	}
 	else if (Event == EMediaEvent::TracksChanged)
@@ -2382,8 +2382,7 @@ void FMediaPlayerFacade::TickFetch(FTimespan DeltaTime, FTimespan Timecode)
 			{
 				if ((FDateTime::UtcNow() - BlockedTime) >= FTimespan::FromSeconds(MEDIAUTILS_MAX_BLOCKONFETCH_SECONDS))
 				{
-					UE_LOG(LogMediaUtils, Verbose, TEXT("PlayerFacade %p: Aborted block on fetch %s after %i seconds"),
-						this,
+					UE_LOG(LogMediaUtils, Verbose, TEXT("PlayerFacade: Aborted block on fetch %s after %i seconds"),
 						*BlockOnRange.GetRange().GetLowerBoundValue().Time.ToString(TEXT("%h:%m:%s.%t")),
 						MEDIAUTILS_MAX_BLOCKONFETCH_SECONDS
 					);
@@ -2393,7 +2392,7 @@ void FMediaPlayerFacade::TickFetch(FTimespan DeltaTime, FTimespan Timecode)
 			}
 			else
 			{
-				UE_LOG(LogMediaUtils, VeryVerbose, TEXT("PlayerFacade %p: Blocking on fetch %s"), this, *BlockOnRange.GetRange().GetLowerBoundValue().Time.ToString(TEXT("%h:%m:%s.%t")));
+				UE_LOG(LogMediaUtils, VeryVerbose, TEXT("PlayerFacade: Blocking on fetch %s"), *BlockOnRange.GetRange().GetLowerBoundValue().Time.ToString(TEXT("%h:%m:%s.%t")));
 
 				Blocked = true;
 				BlockedTime = FDateTime::UtcNow();
@@ -3296,7 +3295,7 @@ void FMediaPlayerFacade::ProcessCaptionSamples(IMediaSamples& Samples, const TRa
 			if (Sample.IsValid() && !CaptionSampleSinks.Enqueue(Sample.ToSharedRef()))
 			{
 #if MEDIAPLAYERFACADE_TRACE_SINKOVERFLOWS
-				UE_LOG(LogMediaUtils, VeryVerbose, TEXT("PlayerFacade %p: Caption sample sink overflow"), this);
+				UE_LOG(LogMediaUtils, VeryVerbose, TEXT("PlayerFacade: Caption sample sink overflow"));
 #endif
 			}
 		}
@@ -3326,7 +3325,7 @@ void FMediaPlayerFacade::ProcessSubtitleSamples(IMediaSamples& Samples, const TR
 			if (Sample.IsValid() && !SubtitleSampleSinks.Enqueue(Sample.ToSharedRef()))
 			{
 #if MEDIAPLAYERFACADE_TRACE_SINKOVERFLOWS
-				UE_LOG(LogMediaUtils, VeryVerbose, TEXT("PlayerFacade %p: Subtitle sample sink overflow"), this);
+				UE_LOG(LogMediaUtils, VeryVerbose, TEXT("PlayerFacade: Subtitle sample sink overflow"));
 #endif
 			}
 		}
@@ -3355,7 +3354,7 @@ void FMediaPlayerFacade::ProcessMetadataSamples(IMediaSamples& Samples, const TR
 			if (Sample.IsValid() && !MetadataSampleSinks.Enqueue(Sample.ToSharedRef()))
 			{
 #if MEDIAPLAYERFACADE_TRACE_SINKOVERFLOWS
-				UE_LOG(LogMediaUtils, VeryVerbose, TEXT("PlayerFacade %p: Metadata sample sink overflow"), this);
+				UE_LOG(LogMediaUtils, VeryVerbose, TEXT("PlayerFacade: Metadata sample sink overflow"));
 #endif
 			}
 		}
@@ -3377,7 +3376,7 @@ void FMediaPlayerFacade::ProcessAudioSamplesV1(IMediaSamples& Samples, TRange<FT
 		if (!AudioSampleSinks.Enqueue(Sample.ToSharedRef()))
 		{
 #if MEDIAPLAYERFACADE_TRACE_SINKOVERFLOWS
-			UE_LOG(LogMediaUtils, VeryVerbose, TEXT("PlayerFacade %p: Audio sample sink overflow"), this);
+			UE_LOG(LogMediaUtils, VeryVerbose, TEXT("PlayerFacade: Audio sample sink overflow"));
 #endif
 		}
 		else
@@ -3415,20 +3414,20 @@ void FMediaPlayerFacade::ProcessVideoSamplesV1(IMediaSamples& Samples, TRange<FT
 			CurrentFrameVideoDisplayTimeStamp = CurrentFrameVideoTimeStamp = Sample->GetTime();
 		}
 
-		UE_LOG(LogMediaUtils, VeryVerbose, TEXT("PlayerFacade %p: Fetched video sample %s"), this, *Sample->GetTime().Time.ToString(TEXT("%h:%m:%s.%t")));
+		UE_LOG(LogMediaUtils, VeryVerbose, TEXT("PlayerFacade: Fetched video sample %s"), *Sample->GetTime().Time.ToString(TEXT("%h:%m:%s.%t")));
 
 		if (VideoSampleSinks.Enqueue(Sample.ToSharedRef()))
 		{
 			if (CurrentRate >= 0.0f)
 			{
 				NextVideoSampleTime = Sample->GetTime().Time + Sample->GetDuration();
-				UE_LOG(LogMediaUtils, VeryVerbose, TEXT("PlayerFacade %p: Next video sample time %s"), this, *NextVideoSampleTime.ToString(TEXT("%h:%m:%s.%t")));
+				UE_LOG(LogMediaUtils, VeryVerbose, TEXT("PlayerFacade: Next video sample time %s"), *NextVideoSampleTime.ToString(TEXT("%h:%m:%s.%t")));
 			}
 		}
 		else
 		{
 #if MEDIAPLAYERFACADE_TRACE_SINKOVERFLOWS
-			UE_LOG(LogMediaUtils, VeryVerbose, TEXT("PlayerFacade %p: Video sample sink overflow"), this);
+			UE_LOG(LogMediaUtils, VeryVerbose, TEXT("PlayerFacade: Video sample sink overflow"));
 #endif
 		}
 	}
@@ -3443,7 +3442,7 @@ void FMediaPlayerFacade::ProcessCaptionSamplesV1(IMediaSamples& Samples, TRange<
 		if (Sample.IsValid() && !CaptionSampleSinks.Enqueue(Sample.ToSharedRef()))
 		{
 #if MEDIAPLAYERFACADE_TRACE_SINKOVERFLOWS
-			UE_LOG(LogMediaUtils, VeryVerbose, TEXT("PlayerFacade %p: Caption sample sink overflow"), this);
+			UE_LOG(LogMediaUtils, VeryVerbose, TEXT("PlayerFacade: Caption sample sink overflow"));
 #endif
 		}
 	}
@@ -3462,7 +3461,7 @@ void FMediaPlayerFacade::ProcessSubtitleSamplesV1(IMediaSamples& Samples, TRange
 			//UE_LOG(LogMediaUtils, Log, TEXT("New caption @%.3f: %s"), Sample->GetTime().Time.GetTotalSeconds(), *Caption);
 
 #if MEDIAPLAYERFACADE_TRACE_SINKOVERFLOWS
-			UE_LOG(LogMediaUtils, VeryVerbose, TEXT("PlayerFacade %p: Subtitle sample sink overflow"), this);
+			UE_LOG(LogMediaUtils, VeryVerbose, TEXT("PlayerFacade: Subtitle sample sink overflow"));
 #endif
 		}
 	}
@@ -3478,7 +3477,7 @@ void FMediaPlayerFacade::ProcessMetadataSamplesV1(IMediaSamples& Samples, TRange
 		if (Sample.IsValid() && !MetadataSampleSinks.Enqueue(Sample.ToSharedRef()))
 		{
 #if MEDIAPLAYERFACADE_TRACE_SINKOVERFLOWS
-			UE_LOG(LogMediaUtils, VeryVerbose, TEXT("PlayerFacade %p: Metadata sample sink overflow"), this);
+			UE_LOG(LogMediaUtils, VeryVerbose, TEXT("PlayerFacade: Metadata sample sink overflow"));
 #endif
 		}
 	}
@@ -3494,134 +3493,134 @@ void FMediaPlayerFacade::ReceiveMediaEvent(EMediaEvent Event)
 	{
 		switch (Event)
 		{
-		case	EMediaEvent::Internal_PurgeVideoSamplesHint:
-		{
-			//
-			// Player asks to attempt to purge older samples in the video output queue it maintains
-			// (ask goes via facade as the player does not have accurate timing info)
-			//
-			TSharedPtr<IMediaPlayer, ESPMode::ThreadSafe> CurrentPlayer = Player;
-
-			if (!CurrentPlayer.IsValid())
+			case EMediaEvent::Internal_PurgeVideoSamplesHint:
 			{
-				return;
-			}
+				//
+				// Player asks to attempt to purge older samples in the video output queue it maintains
+				// (ask goes via facade as the player does not have accurate timing info)
+				//
+				TSharedPtr<IMediaPlayer, ESPMode::ThreadSafe> CurrentPlayer = Player;
+
+				if (!CurrentPlayer.IsValid())
+				{
+					return;
+				}
 
-			// We only support this for V2 timing players
-			check(CurrentPlayer->GetPlayerFeatureFlag(IMediaPlayer::EFeatureFlag::UsePlaybackTimingV2));
+				// We only support this for V2 timing players
+				check(CurrentPlayer->GetPlayerFeatureFlag(IMediaPlayer::EFeatureFlag::UsePlaybackTimingV2));
 
-			// Only do this if we do not block on time ranges
-			if (BlockOnRange.IsSet())
-			{
-				// We do not purge as we do not need max perf, but max reliability to actually get certain frames
-				return;
+				// Only do this if we do not block on time ranges
+				if (BlockOnRange.IsSet())
+				{
+					// We do not purge as we do not need max perf, but max reliability to actually get certain frames
+					return;
+				}
+
+				float Rate = CurrentRate;
+				if (Rate == 0.0f)
+				{
+					return;
+				}
+
+				// Get current playback time
+				// (Note: the delta time is entirely synthetic - we do not pass zero to avoid an empty range, but we do not look far into the future either
+				//        -> after all: we are mainly focused on purging samples up to the current time
+				//  Remarks:
+				//   - this version does not take any estimations from any frame start into account as this is entirely async to the main thread
+				//   - video streams with no audio content will be played using the UE DeltaTime -> so if that stops, the progress of the video stops!
+				//     -> hence we will not see (other then one initial purge) any purging of samples here!
+				// )
+				TRange<FMediaTimeStamp> TimeRange;
+				if (!GetCurrentPlaybackTimeRange(TimeRange, Rate, FTimespan::FromMilliseconds(kOutdatedSamplePurgeRange), true))
+				{
+					return;
+				}
+
+				bool bReverse = (Rate < 0.0f);
+				const float RateFactor = (Rate != 0.0f) ? (1.0f / Rate) : 1.0f;
+
+				// Don't purge frames if the queue is small (to avoid purging if players deliver frames late persistently)
+				uint32 NumPurged = 0;
+				if (CurrentPlayer->GetSamples().NumVideoSamples() >= kMinFramesInVideoQueueToPurge)
+				{
+					NumPurged = CurrentPlayer->GetSamples().PurgeOutdatedVideoSamples(TimeRange.GetLowerBoundValue(), bReverse, FTimespan::FromSeconds(kOutdatedVideoSamplesTolerance * RateFactor));
+				}
+				SET_DWORD_STAT(STAT_MediaUtils_FacadeNumPurgedVideoSamples, NumPurged);
+				INC_DWORD_STAT_BY(STAT_MediaUtils_FacadeTotalPurgedVideoSamples, NumPurged);
+
+				// Take the opportunity to also purge any samples related to video samples directly (and evaluated on the game thread)
+
+				// Captions...
+				NumPurged = 0;
+				if (CurrentPlayer->GetSamples().NumVideoSamples() >= kMinFramesInCaptionQueueToPurge)
+				{
+					NumPurged = CurrentPlayer->GetSamples().PurgeOutdatedCaptionSamples(TimeRange.GetLowerBoundValue(), bReverse, FTimespan::FromSeconds(kOutdatedSubtitleSamplesTolerance * RateFactor));
+				}
+				SET_DWORD_STAT(STAT_MediaUtils_FacadeNumPurgedSubtitleSamples, NumPurged);
+				INC_DWORD_STAT_BY(STAT_MediaUtils_FacadeTotalPurgedSubtitleSamples, NumPurged);
+
+				// Subtitles...
+				NumPurged = 0;
+				if (CurrentPlayer->GetSamples().NumVideoSamples() >= kMinFramesInSubtitleQueueToPurge)
+				{
+					NumPurged = CurrentPlayer->GetSamples().PurgeOutdatedSubtitleSamples(TimeRange.GetLowerBoundValue(), bReverse, FTimespan::FromSeconds(kOutdatedSubtitleSamplesTolerance * RateFactor));
+				}
+				SET_DWORD_STAT(STAT_MediaUtils_FacadeNumPurgedCaptionSamples, NumPurged);
+				INC_DWORD_STAT_BY(STAT_MediaUtils_FacadeTotalPurgedCaptionSamples, NumPurged);
+
+				break;
 			}
 
-			float Rate = CurrentRate;
-			if (Rate == 0.0f)
+			case EMediaEvent::Internal_ResetForDiscontinuity:
 			{
-				return;
+				break;
 			}
-
-			// Get current playback time
-			// (Note: the delta time is entirely synthetic - we do not pass zero to avoid an empty range, but we do not look far into the future either
-			//        -> after all: we are mainly focused on purging samples up to the current time
-			//  Remarks:
-			//   - this version does not take any estimations from any frame start into account as this is entirely async to the main thread
-			//   - video streams with no audio content will be played using the UE DeltaTime -> so if that stops, the progress of the video stops!
-			//     -> hence we will not see (other then one initial purge) any purging of samples here!
-			// )
-			TRange<FMediaTimeStamp> TimeRange;
-			if (!GetCurrentPlaybackTimeRange(TimeRange, Rate, FTimespan::FromMilliseconds(kOutdatedSamplePurgeRange), true))
+			case EMediaEvent::Internal_RenderClockStart:
 			{
-				return;
+				UE_LOG(LogMediaUtils, VeryVerbose, TEXT("PlayerFacade: Render clock shall start"));
+				break;
 			}
-
-			bool bReverse = (Rate < 0.0f);
-			const float RateFactor = (Rate != 0.0f) ? (1.0f / Rate) : 1.0f;
-
-			// Don't purge frames if the queue is small (to avoid purging if players deliver frames late persistently)
-			uint32 NumPurged = 0;
-			if (CurrentPlayer->GetSamples().NumVideoSamples() >= kMinFramesInVideoQueueToPurge)
+			case EMediaEvent::Internal_RenderClockStop:
 			{
-				NumPurged = CurrentPlayer->GetSamples().PurgeOutdatedVideoSamples(TimeRange.GetLowerBoundValue(), bReverse, FTimespan::FromSeconds(kOutdatedVideoSamplesTolerance * RateFactor));
+				UE_LOG(LogMediaUtils, VeryVerbose, TEXT("PlayerFacade: Render clock shall stop"));
+				break;
 			}
-			SET_DWORD_STAT(STAT_MediaUtils_FacadeNumPurgedVideoSamples, NumPurged);
-			INC_DWORD_STAT_BY(STAT_MediaUtils_FacadeTotalPurgedVideoSamples, NumPurged);
 
-			// Take the opportunity to also purge any samples related to video samples directly (and evaluated on the game thread)
-
-			// Captions...
-			NumPurged = 0;
-			if (CurrentPlayer->GetSamples().NumVideoSamples() >= kMinFramesInCaptionQueueToPurge)
+			case EMediaEvent::Internal_VideoSamplesAvailable:
 			{
-				NumPurged = CurrentPlayer->GetSamples().PurgeOutdatedCaptionSamples(TimeRange.GetLowerBoundValue(), bReverse, FTimespan::FromSeconds(kOutdatedSubtitleSamplesTolerance * RateFactor));
+				UE_LOG(LogMediaUtils, VeryVerbose, TEXT("PlayerFacade: Video samples ARE available"));
+				VideoSampleAvailability = 1;
+				break;
 			}
-			SET_DWORD_STAT(STAT_MediaUtils_FacadeNumPurgedSubtitleSamples, NumPurged);
-			INC_DWORD_STAT_BY(STAT_MediaUtils_FacadeTotalPurgedSubtitleSamples, NumPurged);
-
-			// Subtitles...
-			NumPurged = 0;
-			if (CurrentPlayer->GetSamples().NumVideoSamples() >= kMinFramesInSubtitleQueueToPurge)
+			case EMediaEvent::Internal_VideoSamplesUnavailable:
 			{
-				NumPurged = CurrentPlayer->GetSamples().PurgeOutdatedSubtitleSamples(TimeRange.GetLowerBoundValue(), bReverse, FTimespan::FromSeconds(kOutdatedSubtitleSamplesTolerance * RateFactor));
+				UE_LOG(LogMediaUtils, VeryVerbose, TEXT("PlayerFacade: Video samples are NOT available"));
+				VideoSampleAvailability = 0;
+				break;
+			}
+			case EMediaEvent::Internal_AudioSamplesAvailable:
+			{
+				UE_LOG(LogMediaUtils, VeryVerbose, TEXT("PlayerFacade: Audio samples ARE available"));
+				AudioSampleAvailability = 1;
+				break;
+			}
+			case EMediaEvent::Internal_AudioSamplesUnavailable:
+			{
+				UE_LOG(LogMediaUtils, VeryVerbose, TEXT("PlayerFacade: Audio samples are NOT available"));
+				AudioSampleAvailability = 0;
+				break;
 			}
-			SET_DWORD_STAT(STAT_MediaUtils_FacadeNumPurgedCaptionSamples, NumPurged);
-			INC_DWORD_STAT_BY(STAT_MediaUtils_FacadeTotalPurgedCaptionSamples, NumPurged);
-
-			break;
-		}
-
-		case	EMediaEvent::Internal_ResetForDiscontinuity:
-		{
-			break;
-		}
-		case	EMediaEvent::Internal_RenderClockStart:
-		{
-			UE_LOG(LogMediaUtils, VeryVerbose, TEXT("PlayerFacade %p: Render clock shall start"), this);
-			break;
-		}
-		case	EMediaEvent::Internal_RenderClockStop:
-		{
-			UE_LOG(LogMediaUtils, VeryVerbose, TEXT("PlayerFacade %p: Render clock shall stop"), this);
-			break;
-		}
-
-		case	EMediaEvent::Internal_VideoSamplesAvailable:
-		{
-			UE_LOG(LogMediaUtils, VeryVerbose, TEXT("PlayerFacade %p: Video samples ARE available"), this);
-			VideoSampleAvailability = 1;
-			break;
-		}
-		case	EMediaEvent::Internal_VideoSamplesUnavailable:
-		{
-			UE_LOG(LogMediaUtils, VeryVerbose, TEXT("PlayerFacade %p: Video samples are NOT available"), this);
-			VideoSampleAvailability = 0;
-			break;
-		}
-		case	EMediaEvent::Internal_AudioSamplesAvailable:
-		{
-			UE_LOG(LogMediaUtils, VeryVerbose, TEXT("PlayerFacade %p: Audio samples ARE available"), this);
-			AudioSampleAvailability = 1;
-			break;
-		}
-		case	EMediaEvent::Internal_AudioSamplesUnavailable:
-		{
-			UE_LOG(LogMediaUtils, VeryVerbose, TEXT("PlayerFacade %p: Audio samples are NOT available"), this);
-			AudioSampleAvailability = 0;
-			break;
-		}
 
-		default:
-		{
-			UE_LOG(LogMediaUtils, VeryVerbose, TEXT("PlayerFacade %p: Received media event %s"), this, *MediaUtils::EventToString(Event));
-			break;
-		}
+			default:
+			{
+				UE_LOG(LogMediaUtils, VeryVerbose, TEXT("PlayerFacade: Received media event %s"), *MediaUtils::EventToString(Event));
+				break;
+			}
 		}
 	}
 	else
 	{
-		UE_LOG(LogMediaUtils, VeryVerbose, TEXT("PlayerFacade %p: Received media event %s"), this, *MediaUtils::EventToString(Event));
+		UE_LOG(LogMediaUtils, VeryVerbose, TEXT("PlayerFacade: Received media event %s"), *MediaUtils::EventToString(Event));
 		QueuedEvents.Enqueue(Event);
 	}
 }
