#fix TRange<FTimespan>::Inclusive in MediaPlayer
Found with [[nodiscard]]

#rnx

diff --git a/Engine/Source/Runtime/MediaUtils/Private/MediaPlayerFacade.cpp b/Engine/Source/Runtime/MediaUtils/Private/MediaPlayerFacade.cpp
--- a/Engine/Source/Runtime/MediaUtils/Private/MediaPlayerFacade.cpp
+++ b/Engine/Source/Runtime/MediaUtils/Private/MediaPlayerFacade.cpp
@@ -1287,8 +1287,7 @@ void FMediaPlayerFacade::SetBlockOnTime(const FTimespan& Time)
 	}
 	else
 	{
-		TRange<FTimespan> Range;
-		Range.Inclusive(Time, Time);
+		TRange<FTimespan> Range = TRange<FTimespan>::Inclusive(Time, Time);
 		BlockOnRange.SetRange(Range);
 		Player->GetControls().SetBlockingPlaybackHint(true);
 	}
