Media: reworking the timestamp associated sequence index

#rnx

diff --git a/Engine/Plugins/Media/ElectraPlayer/Source/ElectraPlayerPlugin/Private/ElectraPlayerAudioSample.h b/Engine/Plugins/Media/ElectraPlayer/Source/ElectraPlayerPlugin/Private/ElectraPlayerAudioSample.h
--- a/Engine/Plugins/Media/ElectraPlayer/Source/ElectraPlayerPlugin/Private/ElectraPlayerAudioSample.h
+++ b/Engine/Plugins/Media/ElectraPlayer/Source/ElectraPlayerPlugin/Private/ElectraPlayerAudioSample.h
@@ -76,7 +76,7 @@ public:
 	{
 		if (DecoderOutput)
 		{
-			DecoderOutput->SetTime(FDecoderTimeStamp(InTime.Time, InTime.SequenceIndex));
+			DecoderOutput->SetTime(FDecoderTimeStamp(InTime.Time, InTime.GetIndexValue()));
 		}
 	}
 
diff --git a/Engine/Plugins/Media/ElectraPlayer/Source/ElectraProtron/Private/Player/ElectraProtronPlayer.cpp b/Engine/Plugins/Media/ElectraPlayer/Source/ElectraProtron/Private/Player/ElectraProtronPlayer.cpp
--- a/Engine/Plugins/Media/ElectraPlayer/Source/ElectraProtron/Private/Player/ElectraProtronPlayer.cpp
+++ b/Engine/Plugins/Media/ElectraPlayer/Source/ElectraProtron/Private/Player/ElectraProtronPlayer.cpp
@@ -369,7 +369,9 @@ bool FElectraProtronPlayer::Seek(const FTimespan& InTime, const FMediaSeekParams
 	check(InAdditionalParams.NewSequenceIndex.IsSet());
 	int32 SequenceIndexNow = CurrentSequenceIndex;
 	CurrentSequenceIndex = InAdditionalParams.NewSequenceIndex.GetValue();
-	bool bOk = Player->Seek(InTime, CurrentSampleQueue, CurrentSequenceIndex);
+	CurrentSampleQueue->SetMinExpectedNextSequenceIndex(InAdditionalParams.NewSequenceIndex);
+	TOptional<int32> NewLoopIndex = 0;
+	bool bOk = Player->Seek(InTime, CurrentSampleQueue, CurrentSequenceIndex, NewLoopIndex);
 	if (bOk)
 	{
 		// Send a seek complete event even if that is not really true and the seek is in progress or
diff --git a/Engine/Plugins/Media/ElectraPlayer/Source/ElectraProtron/Private/Player/ElectraProtronPlayerControls.cpp b/Engine/Plugins/Media/ElectraPlayer/Source/ElectraProtron/Private/Player/ElectraProtronPlayerControls.cpp
--- a/Engine/Plugins/Media/ElectraPlayer/Source/ElectraProtron/Private/Player/ElectraProtronPlayerControls.cpp
+++ b/Engine/Plugins/Media/ElectraPlayer/Source/ElectraProtron/Private/Player/ElectraProtronPlayerControls.cpp
@@ -149,10 +149,10 @@ void FElectraProtronPlayer::FImpl::HandleSeekRequest(const FSeekRequest& InSeekR
 	UpdateTrackLoader(CodecTypeIndex(ElectraProtronUtils::FCodecInfo::EType::Audio));
 
 	VideoDecoderThread.SetSampleQueue(InSeekRequest.NewSampleQueue);
-	VideoDecoderThread.SetTime(InSeekRequest.NewTime, InSeekRequest.NewSequenceIndex, 0);
+	VideoDecoderThread.SetTime(InSeekRequest.NewTime, InSeekRequest.NewSequenceIndex, InSeekRequest.NewLoopIndex);
 	VideoDecoderThread.ResumeAfterSeek();
 
 	AudioDecoderThread.SetSampleQueue(InSeekRequest.NewSampleQueue);
-	AudioDecoderThread.SetTime(InSeekRequest.NewTime, InSeekRequest.NewSequenceIndex, 0);
+	AudioDecoderThread.SetTime(InSeekRequest.NewTime, InSeekRequest.NewSequenceIndex, InSeekRequest.NewLoopIndex);
 	AudioDecoderThread.ResumeAfterSeek();
 }
diff --git a/Engine/Plugins/Media/ElectraPlayer/Source/ElectraProtron/Private/Player/ElectraProtronPlayerDecoder.cpp b/Engine/Plugins/Media/ElectraPlayer/Source/ElectraProtron/Private/Player/ElectraProtronPlayerDecoder.cpp
--- a/Engine/Plugins/Media/ElectraPlayer/Source/ElectraProtron/Private/Player/ElectraProtronPlayerDecoder.cpp
+++ b/Engine/Plugins/Media/ElectraPlayer/Source/ElectraProtron/Private/Player/ElectraProtronPlayerDecoder.cpp
@@ -186,14 +186,14 @@ uint32 FElectraProtronPlayer::FImpl::FDecoderThread::Run()
 			PendingSeek.Lock.Lock();
 			FTimespan NewTime = PendingSeek.NewTime;
 			int32 NewSequenceIndex = PendingSeek.NewSeqIdx;
-			int32 NewLoopIndex = PendingSeek.NewLoopIdx;
+			TOptional<int32> NewLoopIndex = PendingSeek.NewLoopIdx;
 			PendingSeek.bIsSet = false;
 			PendingSeek.Lock.Unlock();
 
 			TimeLock.Lock();
 			CurrentTime = NewTime;
 			SequenceIndex = NewSequenceIndex;
-			LoopIndex = NewLoopIndex;
+			LoopIndex = NewLoopIndex.Get(LoopIndex);
 			TimeLock.Unlock();
 
 			UpdateTrackIterator(NewTime);
@@ -299,7 +299,7 @@ void FElectraProtronPlayer::FImpl::FDecoderThread::DisconnectSampleBuffer()
 }
 
 
-void FElectraProtronPlayer::FImpl::FDecoderThread::SetTime(FTimespan InTime, int32 InSeqIdx, int32 InLoopIdx)
+void FElectraProtronPlayer::FImpl::FDecoderThread::SetTime(FTimespan InTime, int32 InSeqIdx, TOptional<int32> InLoopIdx)
 {
 	check(IsPausedForSeek());
 	FScopeLock lock(&PendingSeek.Lock);
@@ -868,7 +868,7 @@ void FElectraProtronPlayer::FImpl::FDecoderThread::HandleOutputFrame()
 					{
 						int32 ByteOffsetToFirstSample = 0;
 						FElectraAudioSampleRef AudioSample = Params.AudioSamplePool->AcquireShared();
-						FMediaTimeStamp PTS(InputForCurrentDecoderOutput->Sample->EffectivePTS, FMediaTimeStamp::MakeSequenceIndex(InputForCurrentDecoderOutput->SequenceIndex, InputForCurrentDecoderOutput->LoopIndex));
+						FMediaTimeStamp PTS(InputForCurrentDecoderOutput->Sample->EffectivePTS, InputForCurrentDecoderOutput->SequenceIndex, InputForCurrentDecoderOutput->LoopIndex);
 						FTimespan Dur(InputForCurrentDecoderOutput->Sample->Duration);
 						FTimespan TrimDurAtStart, TrimDurAtEnd;
 						if (AdjustToTime.IsSet())
diff --git a/Engine/Plugins/Media/ElectraPlayer/Source/ElectraProtron/Private/Player/ElectraProtronPlayerImpl.cpp b/Engine/Plugins/Media/ElectraPlayer/Source/ElectraProtron/Private/Player/ElectraProtronPlayerImpl.cpp
--- a/Engine/Plugins/Media/ElectraPlayer/Source/ElectraProtron/Private/Player/ElectraProtronPlayerImpl.cpp
+++ b/Engine/Plugins/Media/ElectraPlayer/Source/ElectraProtron/Private/Player/ElectraProtronPlayerImpl.cpp
@@ -1251,12 +1251,13 @@ bool FElectraProtronPlayer::FImpl::IsLooping()
 	return SharedPlayParams->bShouldLoop;
 }
 
-bool FElectraProtronPlayer::FImpl::Seek(const FTimespan& InTime, const TSharedPtr<FMediaSamples, ESPMode::ThreadSafe>& InNewSampleQueue, uint32 InNewSequenceIndex)
+bool FElectraProtronPlayer::FImpl::Seek(const FTimespan& InTime, const TSharedPtr<FMediaSamples, ESPMode::ThreadSafe>& InNewSampleQueue, int32 InNewSequenceIndex, const TOptional<int32>& InNewLoopIndex)
 {
 	FSeekRequest sr;
 	sr.NewSampleQueue = InNewSampleQueue;
 	sr.NewTime = InTime;
 	sr.NewSequenceIndex = InNewSequenceIndex;
+	sr.NewLoopIndex = InNewLoopIndex;
 	SeekRequestLock.Lock();
 	PendingSeekRequest = MoveTemp(sr);
 	SeekRequestLock.Unlock();
diff --git a/Engine/Plugins/Media/ElectraPlayer/Source/ElectraProtron/Private/Player/ElectraProtronPlayerImpl.h b/Engine/Plugins/Media/ElectraPlayer/Source/ElectraProtron/Private/Player/ElectraProtronPlayerImpl.h
--- a/Engine/Plugins/Media/ElectraPlayer/Source/ElectraProtron/Private/Player/ElectraProtronPlayerImpl.h
+++ b/Engine/Plugins/Media/ElectraPlayer/Source/ElectraProtron/Private/Player/ElectraProtronPlayerImpl.h
@@ -72,7 +72,7 @@ public:
 	FTimespan GetTime();
 	bool SetLooping(bool bInLooping);
 	bool IsLooping();
-	bool Seek(const FTimespan& InTime, const TSharedPtr<FMediaSamples, ESPMode::ThreadSafe>& InNewSampleQueue, uint32 InNewSequenceIndex);
+	bool Seek(const FTimespan& InTime, const TSharedPtr<FMediaSamples, ESPMode::ThreadSafe>& InNewSampleQueue, int32 InNewSequenceIndex, const TOptional<int32>& InNewLoopIndex);
 
 	void TickFetch(FTimespan InDeltaTime, FTimespan InTimecode);
 	void TickInput(FTimespan InDeltaTime, FTimespan InTimecode);
@@ -107,6 +107,7 @@ private:
 		TSharedPtr<FMediaSamples, ESPMode::ThreadSafe> NewSampleQueue;
 		FTimespan NewTime;
 		int32 NewSequenceIndex = 0;
+		TOptional<int32> NewLoopIndex;
 	};
 
 	struct FWorkerThreadMessage
@@ -365,7 +366,7 @@ private:
 		void PauseForSeek();
 		void ResumeAfterSeek();
 		bool IsPausedForSeek();
-		void SetTime(FTimespan InTime, int32 InSeqIdx, int32 InLoopIdx);
+		void SetTime(FTimespan InTime, int32 InSeqIdx, TOptional<int32> InLoopIdx);
 		void SetEstimatedPlaybackTime(FTimespan InTime);
 
 		void Flush(const TSharedPtr<FMediaEvent, ESPMode::ThreadSafe>& InFlushedSignal);
@@ -405,7 +406,7 @@ private:
 			FCriticalSection Lock;
 			FTimespan NewTime;
 			int32 NewSeqIdx = 0;
-			int32 NewLoopIdx = 0;
+			TOptional<int32> NewLoopIdx;
 			bool bIsSet = false;
 		};
 
diff --git a/Engine/Plugins/Media/ElectraUtil/Source/ElectraSamples/Private/ElectraTextureSampleBase.cpp b/Engine/Plugins/Media/ElectraUtil/Source/ElectraSamples/Private/ElectraTextureSampleBase.cpp
--- a/Engine/Plugins/Media/ElectraUtil/Source/ElectraSamples/Private/ElectraTextureSampleBase.cpp
+++ b/Engine/Plugins/Media/ElectraUtil/Source/ElectraSamples/Private/ElectraTextureSampleBase.cpp
@@ -248,7 +248,7 @@ void IElectraTextureSampleBase::SetTime(const FMediaTimeStamp& InTime)
 {
 	if (VideoDecoderOutput)
 	{
-		VideoDecoderOutput->SetTime(FDecoderTimeStamp(InTime.Time, InTime.SequenceIndex));
+		VideoDecoderOutput->SetTime(FDecoderTimeStamp(InTime.Time, InTime.GetIndexValue()));
 	}
 }
 
diff --git a/Engine/Plugins/Media/ImgMedia/Source/ImgMedia/Internal/Loader/ImgMediaLoader.h b/Engine/Plugins/Media/ImgMedia/Source/ImgMedia/Internal/Loader/ImgMediaLoader.h
--- a/Engine/Plugins/Media/ImgMedia/Source/ImgMedia/Internal/Loader/ImgMediaLoader.h
+++ b/Engine/Plugins/Media/ImgMedia/Source/ImgMedia/Internal/Loader/ImgMediaLoader.h
@@ -665,12 +665,13 @@ private:
 	bool UseGlobalCache;
 
 	/** State related to "queue style" frame access functions */
-	struct
+	struct FQueuedSampleFetch
 	{
 		FMediaTimeStamp LastTimeStamp;	// Current playhead / last delivered sample PTS
 		FTimespan LastDuration;			// Duration of a frame / zero after seek or playback start
 		int32 LoopIndex;				// Current loop index to be encoded into sequence index of timestamp
-	} QueuedSampleFetch;
+	};
+	FQueuedSampleFetch QueuedSampleFetch;
 
 private:
 
diff --git a/Engine/Plugins/Media/ImgMedia/Source/ImgMedia/Internal/Player/ImgMediaPlayer.h b/Engine/Plugins/Media/ImgMedia/Source/ImgMedia/Internal/Player/ImgMediaPlayer.h
--- a/Engine/Plugins/Media/ImgMedia/Source/ImgMedia/Internal/Player/ImgMediaPlayer.h
+++ b/Engine/Plugins/Media/ImgMedia/Source/ImgMedia/Internal/Player/ImgMediaPlayer.h
@@ -70,7 +70,6 @@ public:
 	virtual void TickInput(FTimespan DeltaTime, FTimespan Timecode) override;
 	virtual bool FlushOnSeekStarted() const;
 	virtual bool FlushOnSeekCompleted() const;
-	virtual void ProcessVideoSamples() override;
 	virtual bool GetPlayerFeatureFlag(EFeatureFlag flag) const override;
 
 protected:
@@ -114,6 +113,7 @@ protected:
 	virtual bool FetchVideo(TRange<FTimespan> TimeRange, TSharedPtr<IMediaTextureSample, ESPMode::ThreadSafe>& OutSample) override;
 	virtual void FlushSamples() override;
 
+	virtual void SetMinExpectedNextSequenceIndex(TOptional<int32> InNextSequenceIndex) override;
 	virtual EFetchBestSampleResult FetchBestVideoSampleForTimeRange(const TRange<FMediaTimeStamp>& TimeRange, TSharedPtr<IMediaTextureSample, ESPMode::ThreadSafe>& OutSample, bool bReverse, bool bConsistentResult) override;
 	virtual bool PeekVideoSampleTime(FMediaTimeStamp & TimeStamp) override;
 	virtual bool DiscardVideoSamples(const TRange<FMediaTimeStamp>& TimeRange, bool bReverse) override;
diff --git a/Engine/Plugins/Media/ImgMedia/Source/ImgMedia/Private/Loader/ImgMediaLoader.cpp b/Engine/Plugins/Media/ImgMedia/Source/ImgMedia/Private/Loader/ImgMediaLoader.cpp
--- a/Engine/Plugins/Media/ImgMedia/Source/ImgMedia/Private/Loader/ImgMediaLoader.cpp
+++ b/Engine/Plugins/Media/ImgMedia/Source/ImgMedia/Private/Loader/ImgMediaLoader.cpp
@@ -188,8 +188,7 @@ FImgMediaLoader::FImgMediaLoader(const TSharedRef<FImgMediaScheduler, ESPMode::T
 #endif
 {
 	ResetFetchLogic();
-	QueuedSampleFetch.LastTimeStamp.Time = FTimespan::Zero();
-	QueuedSampleFetch.LastTimeStamp.SequenceIndex = 0;
+	QueuedSampleFetch.LastTimeStamp.SetTimeAndIndexValue(FTimespan::Zero(), 0);
 	QueuedSampleFetch.LastDuration = FTimespan::Zero();
 
 	UE_LOG(LogImgMedia, Verbose, TEXT("Loader %p: Created"), this);
@@ -271,7 +270,7 @@ TSharedPtr<FImgMediaTextureSample, ESPMode::ThreadSafe> FImgMediaLoader::GetFram
 	FScopeLock ScopeLock(&CriticalSection);
 
 	const TSharedPtr<FImgMediaFrame, ESPMode::ThreadSafe>* Frame = GetCachedFrame(FrameIndex);
-	
+
 	if (Frame == nullptr)
 	{
 		return nullptr;
@@ -340,13 +339,14 @@ void FImgMediaLoader::Flush()
 {
 	ResetFetchLogic();
 	// Adjust last timestamp / play head so it is within the current, reset sequence index range
-	QueuedSampleFetch.LastTimeStamp.SequenceIndex = FMediaTimeStamp::MakeSequenceIndex(0, 0);
+	QueuedSampleFetch.LastTimeStamp.SetSequenceIndex(0);
+	QueuedSampleFetch.LastTimeStamp.SetLoopIndex(0);
 }
 
 void FImgMediaLoader::Seek(const FMediaTimeStamp& SeekTarget, bool bReverse)
 {
 	ResetFetchLogic();
-	QueuedSampleFetch.LastTimeStamp = { FMediaTimeStamp(FrameNumberToTime(TimeToFrameNumber(SeekTarget.Time)), SeekTarget.SequenceIndex) };
+	QueuedSampleFetch.LastTimeStamp.SetTimeAndIndexValue(FrameNumberToTime(TimeToFrameNumber(SeekTarget.Time)), SeekTarget.GetIndexValue());
 	QueuedSampleFetch.LastDuration = FTimespan::Zero();
 }
 
@@ -391,7 +391,7 @@ IMediaSamples::EFetchBestSampleResult FImgMediaLoader::FetchBestVideoSampleForTi
 	{
 		// A few sanity checks
 		check(InTimeRange.GetLowerBoundValue() <= InTimeRange.GetUpperBoundValue());
-		check(FMediaTimeStamp::GetPrimaryIndex(InTimeRange.GetLowerBoundValue().SequenceIndex) == FMediaTimeStamp::GetPrimaryIndex(InTimeRange.GetUpperBoundValue().SequenceIndex));
+		check(InTimeRange.GetLowerBoundValue().GetSequenceIndex() == InTimeRange.GetUpperBoundValue().GetSequenceIndex());
 		check(QueuedSampleFetch.LastTimeStamp.IsValid());
 
 		auto TimeRange = InTimeRange;
@@ -399,40 +399,11 @@ IMediaSamples::EFetchBestSampleResult FImgMediaLoader::FetchBestVideoSampleForTi
 		// To avoid odd behavior with out of range time values, clamp them down just in case
 		if (!bIsLoopingEnabled)
 		{
-			TimeRange.SetLowerBoundValue(FMediaTimeStamp(FMath::Clamp(TimeRange.GetLowerBoundValue().Time, FTimespan::Zero(), SequenceDuration), TimeRange.GetLowerBoundValue().SequenceIndex));
-			TimeRange.SetUpperBoundValue(FMediaTimeStamp(FMath::Clamp(TimeRange.GetUpperBoundValue().Time, FTimespan::Zero(), SequenceDuration), TimeRange.GetUpperBoundValue().SequenceIndex));
-		}
-
-		// Clamp the current range requested against the last known timestamp to ensure we return the "next" sample at all times...
-		if (PlayRate >= 0.0f)
-		{
-			FMediaTimeStamp NextTimeStamp = QueuedSampleFetch.LastTimeStamp + QueuedSampleFetch.LastDuration;
-			if (NextTimeStamp >= TimeRange.GetUpperBoundValue())
-			{
-				return IMediaSamples::EFetchBestSampleResult::NoSample;
-			}
-
-			if (TimeRange.GetLowerBoundValue() < NextTimeStamp)
-			{
-				TimeRange = TRange<FMediaTimeStamp>(NextTimeStamp, InTimeRange.GetUpperBoundValue());
-			}
+			TimeRange.SetLowerBoundValue(FMediaTimeStamp(FMath::Clamp(TimeRange.GetLowerBoundValue().Time, FTimespan::Zero(), SequenceDuration), TimeRange.GetLowerBoundValue().GetIndexValue()));
+			TimeRange.SetUpperBoundValue(FMediaTimeStamp(FMath::Clamp(TimeRange.GetUpperBoundValue().Time, FTimespan::Zero(), SequenceDuration), TimeRange.GetUpperBoundValue().GetIndexValue()));
 		}
-		else
-		{
-			FMediaTimeStamp NextTimeStampAtEnd = (QueuedSampleFetch.LastTimeStamp - QueuedSampleFetch.LastDuration) + FTimespan::FromSeconds(SequenceFrameRate.AsInterval());
-			if (NextTimeStampAtEnd < TimeRange.GetLowerBoundValue())
-			{
-				return IMediaSamples::EFetchBestSampleResult::NoSample;
-			}
-
-			if (TimeRange.GetUpperBoundValue() > NextTimeStampAtEnd)
-			{
-				TimeRange = TRange<FMediaTimeStamp>(InTimeRange.GetLowerBoundValue(), NextTimeStampAtEnd);
-			}
-		}
-
 		FTimespan StartTime = TimeRange.GetLowerBoundValue().Time;
-		int32 StartLoopIndex = FMediaTimeStamp::GetSecondaryIndex(TimeRange.GetLowerBoundValue().SequenceIndex);
+		int32 StartLoopIndex = TimeRange.GetLowerBoundValue().GetLoopIndex();
 		FTimespan EndTime = TimeRange.GetUpperBoundValue().Time;
 		int32 EndLoopIndex;
 
@@ -440,12 +411,12 @@ IMediaSamples::EFetchBestSampleResult FImgMediaLoader::FetchBestVideoSampleForTi
 		if (EndTime > FTimespan::Zero())
 		{
 			EndTime -= FTimespan(1);
-			EndLoopIndex = FMediaTimeStamp::GetSecondaryIndex(TimeRange.GetUpperBoundValue().SequenceIndex);
+			EndLoopIndex = TimeRange.GetUpperBoundValue().GetLoopIndex();
 		}
 		else
 		{
 			EndTime = SequenceDuration - FTimespan(1);
-			EndLoopIndex = FMediaTimeStamp::GetSecondaryIndex(TimeRange.GetUpperBoundValue().SequenceIndex) - 1;
+			EndLoopIndex = TimeRange.GetUpperBoundValue().GetLoopIndex() - 1;
 		}
 
 		int32 NumFrames = GetNumImages();
@@ -578,7 +549,9 @@ IMediaSamples::EFetchBestSampleResult FImgMediaLoader::FetchBestVideoSampleForTi
 			}
 
 			// Track state & setup sample for caller...
-			QueuedSampleFetch.LastTimeStamp = FMediaTimeStamp(SampleTime, FMediaTimeStamp::MakeSequenceIndex(FMediaTimeStamp::GetPrimaryIndex(TimeRange.GetLowerBoundValue().SequenceIndex), QueuedSampleFetch.LoopIndex));
+			QueuedSampleFetch.LastTimeStamp.SetTime(SampleTime);
+			QueuedSampleFetch.LastTimeStamp.SetSequenceIndex(TimeRange.GetLowerBoundValue().GetSequenceIndex());
+			QueuedSampleFetch.LastTimeStamp.SetLoopIndex(QueuedSampleFetch.LoopIndex);
 			QueuedSampleFetch.LastDuration = FTimespan::FromSeconds(Frame->Get()->GetFrameRate().AsInterval());
 
 			// We are clear to return it as new result... Make a sample & initialize it...
@@ -626,7 +599,7 @@ IMediaSamples::EFetchBestSampleResult FImgMediaLoader::FetchBestVideoSampleForTi
 	return IMediaSamples::EFetchBestSampleResult::NoSample;
 }
 
- 
+
 bool FImgMediaLoader::PeekVideoSampleTime(FMediaTimeStamp &TimeStamp, bool bIsLoopingEnabled, float PlayRate, const FTimespan& CurrentTime)
 {
 	// Strictly speaking we never actual return a "frame" here as we do NOT involve any requests to the cache -- we simply assume we could deliver what a "queue" would contain on demand later...
@@ -637,19 +610,19 @@ bool FImgMediaLoader::PeekVideoSampleTime(FMediaTimeStamp &TimeStamp, bool bIsLo
 		if (PlayRate >= 0.0f)
 		{
 			TimeStamp = QueuedSampleFetch.LastTimeStamp + QueuedSampleFetch.LastDuration;
-			if (TimeStamp.Time >= SequenceDuration)
+			while(TimeStamp.Time >= SequenceDuration)
 			{
 				TimeStamp.Time -= SequenceDuration;
-				TimeStamp.SequenceIndex = FMediaTimeStamp::AdjustSecondaryIndex(QueuedSampleFetch.LastTimeStamp.SequenceIndex, 1);
+				TimeStamp.AdjustLoopIndex(1);
 			}
 		}
 		else
 		{
 			TimeStamp = QueuedSampleFetch.LastTimeStamp - QueuedSampleFetch.LastDuration;
-			if (TimeStamp.Time.GetTicks() < 0)
+			while(TimeStamp.Time.GetTicks() < 0)
 			{
 				TimeStamp.Time += SequenceDuration;
-				TimeStamp.SequenceIndex = FMediaTimeStamp::AdjustSecondaryIndex(QueuedSampleFetch.LastTimeStamp.SequenceIndex, -1);
+				TimeStamp.AdjustLoopIndex(-1);
 			}
 		}
 
@@ -714,14 +687,14 @@ IQueuedWork* FImgMediaLoader::GetWork()
 	}
 
 	FImgMediaLoaderWork* Work = (WorkPool.Num() > 0) ? WorkPool.Pop() : new FImgMediaLoaderWork(AsShared(), Reader.ToSharedRef());
-	
+
 	// Get the existing frame so we can add the mip level to it.
 	TSharedPtr<FImgMediaFrame, ESPMode::ThreadSafe> ExistingFrame;
 	if (const TSharedPtr<FImgMediaFrame, ESPMode::ThreadSafe>* ExistingFramePtr = GetCachedFrame(FrameNumber))
 	{
 		ExistingFrame = *ExistingFramePtr;
 	}
-	
+
 	// Set up work.
 	Work->Initialize(NextFrameNumberJobID, FrameNumber, DesiredMipsAndTiles, ExistingFrame);
 	QueuedFrameNumbers.Emplace(FrameNumber, NextFrameNumberJobID++);
@@ -743,12 +716,12 @@ void FImgMediaLoader::Initialize(const FString& SequencePath, const FFrameRate&
 	check(!Initialized); // reinitialization not allowed for now
 
 	FImgMediaFrameInfo FirstFrameInfo;
-	
+
 	if (LoadSequence(SequencePath, FrameRateOverride, Loop, FirstFrameInfo))
 	{
 		WarmupSequence(FirstFrameInfo, Loop);
 	}
-	
+
 	FPlatformMisc::MemoryBarrier();
 
 	Initialized = true;
@@ -941,7 +914,7 @@ bool FImgMediaLoader::LoadSequence(const FString& SequencePath, const FFrameRate
 
 	// initialize loader
 	const FPlatformMemoryStats Stats = FPlatformMemory::GetStats();
-	
+
 	// We always need to cache at least 2 frames - current and the future frame.
 	const int32 MinNumberOfFramesToLoad = 2;
 
@@ -950,7 +923,7 @@ bool FImgMediaLoader::LoadSequence(const FString& SequencePath, const FFrameRate
 	{
 		// Load frames that are a certain amount of time ahead of us.
 		float TimeToLookAhead = SmartCacheSettings.TimeToLookAhead;
-		
+
 		NumFramesToLoad = TimeToLookAhead / SequenceFrameRate.AsInterval();
 		// Clamp min frame to load to 1 so we always try to load something.
 		int32 NumImages = GetNumImages();
@@ -990,7 +963,7 @@ bool FImgMediaLoader::LoadSequence(const FString& SequencePath, const FFrameRate
 void FImgMediaLoader::WarmupSequence(const FImgMediaFrameInfo& InFirstFrameInfo, bool Loop)
 {
 	// Giving our reader a chance to handle RAM allocation.
-	// Not all readers use this, only those that need to handle large files 
+	// Not all readers use this, only those that need to handle large files
 	// or need to be as efficient as possible.
 	Reader->PreAllocateMemoryPool(NumFramesToLoad, InFirstFrameInfo);
 
@@ -1077,7 +1050,7 @@ void FImgMediaLoader::FindMips(const FString& SequencePath)
 	// Get parent directory.
 	FString BaseName = FPaths::GetCleanFilename(SequenceDir);
 	FString ParentDir = FPaths::GetPath(SequenceDir);
-	
+
 	// Is this a mipmap level, i.e. something like 256x256?
 	int32 Index = 0;
 	bool FoundDelimeter = SequenceDir.FindLastChar(TEXT('x'), Index);
@@ -1126,7 +1099,7 @@ void FImgMediaLoader::FindMips(const FString& SequencePath)
 			// Make sure we have the same number of mip files over all levels.
 			if (MipFiles.Num() != GetNumImages())
 			{
-				UE_LOG(LogImgMedia, Error, TEXT("Loader %p: Found %d images, expected %d, for %s"), 
+				UE_LOG(LogImgMedia, Error, TEXT("Loader %p: Found %d images, expected %d, for %s"),
 					this, MipFiles.Num(), GetNumImages(), *MipDir);
 				break;
 			}
@@ -1211,7 +1184,7 @@ void FImgMediaLoader::Update(int32 PlayHeadFrame, float PlayRate, bool Loop)
 		}
 
 		// alternate between look ahead and look behind
-		
+
 		while ((FramesToLoad.Num() < NumFramesToLoad) &&
 			((LoadAheadCount > 0) || (LoadBehindCount > 0)))
 		{
@@ -1301,7 +1274,7 @@ void FImgMediaLoader::Update(int32 PlayHeadFrame, float PlayRate, bool Loop)
 			UE_LOG(LogImgMedia, Verbose, TEXT("Loader %p: Removed Frame %i"), this, FrameNumber);
 			QueuedFrameNumbers.RemoveAtSwap(Idx);
 			Reader->CancelFrame(FrameNumber);
-			
+
 			// We failed to load this frame... are we throttling?
 			if (bIsSkipFramesEnabled)
 			{
@@ -1346,7 +1319,7 @@ void FImgMediaLoader::Update(int32 PlayHeadFrame, float PlayRate, bool Loop)
 				UE_LOG(LogImgMedia, VeryVerbose, TEXT("Loader %p: Invalidating frame %i at playhead %i due to missing mips or tiles."), this, FrameNumber, PlayHeadFrame);
 			}
 		}
-		
+
 		if (NeedFrame)
 		{
 			bool FoundFrame = false;
@@ -1411,7 +1384,7 @@ bool FImgMediaLoader::TryAddEmptyFrame(int32 FrameNumber)
 		const int32 NumChannels = 3;
 		const int32 PixelSize = sizeof(uint16) * NumChannels;
 		TSharedPtr<FImgMediaFrame, ESPMode::ThreadSafe> Frame = MakeShared<FImgMediaFrame>();
-		
+
 		FImgMediaFrameInfo FrameInfo = {};
 		FrameInfo.Dim = SequenceDim;
 		FrameInfo.FrameRate = SequenceFrameRate;
@@ -1445,7 +1418,7 @@ void FImgMediaLoader::NotifyWorkComplete(FImgMediaLoaderWork& CompletedWork, int
 			SkipFramesCounter--;
 		}
 	}
-	
+
 	// Make sure the frame is no longer cancelled, otherwise we might end up cancelling this frame
 	// later on inadvertently!
 	Reader->UncancelFrame(FrameNumber);
@@ -1486,7 +1459,7 @@ void FImgMediaLoader::AddFrameToCache(int32 FrameNumber, const TSharedPtr<FImgMe
 void FImgMediaLoader::GetDesiredMipTiles(int32 FrameIndex, TMap<int32, FImgMediaTileSelection>& OutMipsAndTiles)
 {
 	// note: While currently unused, FrameIndex could become useful for recorded frames.
-	
+
 	if(HasMipsOrTilesEnabled() && MipMapInfo->HasObjects())
 	{
 		OutMipsAndTiles = MipMapInfo->GetVisibleTiles();
@@ -1515,7 +1488,7 @@ int32 FImgMediaLoader::GetDesiredMinimumMipLevelToUpscale()
 FTimespan FImgMediaLoader::ModuloTime(FTimespan Time)
 {
 	bool IsNegative = Time < FTimespan::Zero();
-	
+
 	FTimespan NewTime = Time % SequenceDuration;
 	if (IsNegative)
 	{
@@ -1540,7 +1513,7 @@ bool FImgMediaLoader::GetNumberAtEndOfString(int32 &Number, const FString& Strin
 			break;
 		}
 	}
-	
+
 	// Did we find one?
 	if (Index >= 0)
 	{
@@ -1565,7 +1538,7 @@ bool FImgMediaLoader::GetNumberAtEndOfString(int32 &Number, const FString& Strin
 			Number = FCString::Atoi(*NumberString);
 		}
 	}
-	
+
 	return bFoundNumber;
 }
 
@@ -1575,7 +1548,7 @@ int32 FImgMediaLoader::GetSkipFrame(int32 FrameIndex)
 	int32 Mult = (1 << SkipFramesLevel);
 	int32 Mask = ~(Mult - 1);
 	FrameIndex = (FrameIndex & Mask) + (Mult - 1);
-	
+
 	const int32 NumImagePaths = GetNumImages();
 	if (FrameIndex >= NumImagePaths)
 	{
diff --git a/Engine/Plugins/Media/ImgMedia/Source/ImgMedia/Private/Player/ImgMediaPlayer.cpp b/Engine/Plugins/Media/ImgMedia/Source/ImgMedia/Private/Player/ImgMediaPlayer.cpp
--- a/Engine/Plugins/Media/ImgMedia/Source/ImgMedia/Private/Player/ImgMediaPlayer.cpp
+++ b/Engine/Plugins/Media/ImgMedia/Source/ImgMedia/Private/Player/ImgMediaPlayer.cpp
@@ -359,10 +359,6 @@ bool FImgMediaPlayer::FlushOnSeekCompleted() const
 	return false;
 }
 
-void FImgMediaPlayer::ProcessVideoSamples()
-{
-}
-
 //-----------------------------------------------------------------------------
 /**
 	Get special feature flags states
@@ -535,11 +531,10 @@ bool FImgMediaPlayer::Seek(const FTimespan& Time, const FMediaSeekParams& InAddi
 	}
 
 	CurrentTime = Time;
-	check(InAdditionalParams.NewSequenceIndex.IsSet());
-	CurrentSeekIndex = InAdditionalParams.NewSequenceIndex.GetValue();
+	CurrentSeekIndex = InAdditionalParams.NewSequenceIndex.Get(0);
 	if (Loader.IsValid())
 	{
-		Loader->Seek(FMediaTimeStamp(CurrentTime, FMediaTimeStamp::MakeSequenceIndex(CurrentSeekIndex, 0)), LastNonZeroRate < 0.0f);
+		Loader->Seek(FMediaTimeStamp(CurrentTime, CurrentSeekIndex, 0), LastNonZeroRate < 0.0f);
 	}
 
 	if (CurrentState == EMediaState::Paused)
@@ -664,6 +659,10 @@ bool FImgMediaPlayer::DiscardVideoSamples(const TRange<FMediaTimeStamp>& TimeRan
 	return false;
 }
 
+void FImgMediaPlayer::SetMinExpectedNextSequenceIndex(TOptional<int32> InNextSequenceIndex)
+{
+}
+
 
 IMediaSamples::EFetchBestSampleResult FImgMediaPlayer::FetchBestVideoSampleForTimeRange(const TRange<FMediaTimeStamp>& TimeRange, TSharedPtr<IMediaTextureSample, ESPMode::ThreadSafe>& OutSample, bool bReverse, bool bConsistentResult)
 {
diff --git a/Engine/Plugins/Media/MfMedia/Source/MfMedia/Private/Player/MfMediaTracks.cpp b/Engine/Plugins/Media/MfMedia/Source/MfMedia/Private/Player/MfMediaTracks.cpp
--- a/Engine/Plugins/Media/MfMedia/Source/MfMedia/Private/Player/MfMediaTracks.cpp
+++ b/Engine/Plugins/Media/MfMedia/Source/MfMedia/Private/Player/MfMediaTracks.cpp
@@ -23,7 +23,11 @@
 #endif
 
 #define MFMEDIATRACKS_TRACE_SAMPLES 0
+#if defined(MFMEDIA_PLATFORM_SUPPORTS_ASYNC_TRACK_READER) && MFMEDIA_PLATFORM_SUPPORTS_ASYNC_TRACK_READER !=0
 #define MFMEDIATRACKS_USE_ASYNCREADER 1
+#else
+#define MFMEDIATRACKS_USE_ASYNCREADER 0
+#endif
 #define MFMEDIATRACKS_USE_DXVA 0			// not implemented yet
 
 #define LOCTEXT_NAMESPACE "FMfMediaSession"
@@ -283,7 +287,7 @@ void FMfMediaTracks::ProcessSample(IMFSample* Sample, HRESULT Status, DWORD Stre
 
 		if (Track.StreamIndex == StreamIndex)
 		{
-			if (Samples->NumAudio() >= FMediaPlayerQueueDepths::MaxAudioSinkDepth)
+			if (Samples->NumAudioSamples() >= FMediaPlayerQueueDepths::MaxAudioSinkDepth)
 			{
 				return;
 			}
@@ -327,7 +331,7 @@ void FMfMediaTracks::ProcessSample(IMFSample* Sample, HRESULT Status, DWORD Stre
 
 		if (Track.StreamIndex == StreamIndex)
 		{
-			if (Samples->NumCaption() >= FMediaPlayerQueueDepths::MaxCaptionSinkDepth)
+			if (Samples->NumCaptionSamples() >= FMediaPlayerQueueDepths::MaxCaptionSinkDepth)
 			{
 				return;
 			}
@@ -1437,14 +1441,14 @@ bool FMfMediaTracks::RequestSample(DWORD StreamIndex)
 {
 	if (SelectedAudioTrack != INDEX_NONE && AudioTracks[SelectedAudioTrack].StreamIndex == StreamIndex)
 	{
-		if (Samples->NumAudio() >= FMediaPlayerQueueDepths::MaxAudioSinkDepth)
+		if (Samples->NumAudioSamples() >= FMediaPlayerQueueDepths::MaxAudioSinkDepth)
 		{
 			return false;
 		}
 	}
 	else if (SelectedCaptionTrack != INDEX_NONE && CaptionTracks[SelectedCaptionTrack].StreamIndex == StreamIndex)
 	{
-		if (Samples->NumCaption() >= FMediaPlayerQueueDepths::MaxCaptionSinkDepth)
+		if (Samples->NumCaptionSamples() >= FMediaPlayerQueueDepths::MaxCaptionSinkDepth)
 		{
 			return false;
 		}
diff --git a/Engine/Plugins/Media/WebMMedia/Source/WebMMedia/Private/Player/WebMMediaPlayer.cpp b/Engine/Plugins/Media/WebMMedia/Source/WebMMedia/Private/Player/WebMMediaPlayer.cpp
--- a/Engine/Plugins/Media/WebMMedia/Source/WebMMedia/Private/Player/WebMMediaPlayer.cpp
+++ b/Engine/Plugins/Media/WebMMedia/Source/WebMMedia/Private/Player/WebMMediaPlayer.cpp
@@ -198,7 +198,7 @@ void FWebMMediaPlayer::TickFetch(FTimespan DeltaTime, FTimespan Timecode)
 	if (bNoMoreToRead)
 	{
 		// Check if we finished playing
-		if (Samples->NumVideoSamples() == 0 && Samples->NumAudio() == 0)
+		if (Samples->NumVideoSamples() == 0 && Samples->NumAudioSamples() == 0)
 		{
 			if (bLooping)
 			{
diff --git a/Engine/Plugins/Media/WmfMedia/Source/WmfMedia/Private/Player/WmfMediaTracks.cpp b/Engine/Plugins/Media/WmfMedia/Source/WmfMedia/Private/Player/WmfMediaTracks.cpp
--- a/Engine/Plugins/Media/WmfMedia/Source/WmfMedia/Private/Player/WmfMediaTracks.cpp
+++ b/Engine/Plugins/Media/WmfMedia/Source/WmfMedia/Private/Player/WmfMediaTracks.cpp
@@ -577,6 +577,7 @@ void FWmfMediaTracks::FlushSamples()
 
 			SeekIndex = 0;
 			LoopIndex = 0;
+			NextMinSequenceIndex.Reset();
 			}
 
 			Session->Flush();
@@ -597,6 +598,27 @@ IMediaSamples::EFetchBestSampleResult FWmfMediaTracks::FetchBestVideoSampleForTi
 	return Ret;
 }
 
+void FWmfMediaTracks::SetMinExpectedNextSequenceIndex(TOptional<int32> InNextSequenceIndex)
+{
+	FScopeLock Lock(&CriticalSection);
+	if (!InNextSequenceIndex.IsSet())
+	{
+		NextMinSequenceIndex.Reset();
+		return;
+	}
+	if (NextMinSequenceIndex == InNextSequenceIndex)
+	{
+		return;
+	}
+	NextMinSequenceIndex = InNextSequenceIndex;
+	int32 nsi = InNextSequenceIndex.GetValue();
+
+	AudioSampleQueue.PurgeUntilSequenceIndex(nsi);
+	MetadataSampleQueue.PurgeUntilSequenceIndex(nsi);
+	CaptionSampleQueue.PurgeUntilSequenceIndex(nsi);
+	VideoSampleQueue.PurgeUntilSequenceIndex(nsi);
+}
+
 
 bool FWmfMediaTracks::PeekVideoSampleTime(FMediaTimeStamp & TimeStamp)
 {
@@ -1243,10 +1265,16 @@ bool FWmfMediaTracks::AddTrackToTopology(const FTrack& Track, IMFTopology& Topol
 	{
 		VideoHardwareVideoDecodingSamplePool = MakeShared<FWmfMediaHardwareVideoDecodingTextureSamplePool>();
 
-		MediaStreamSink->SetMediaSamplePoolAndQueue(VideoHardwareVideoDecodingSamplePool, &VideoSampleQueue, [this](FTimespan Time, EMediaTrackType TrackType) -> FMediaTimeStamp
-		{
-			return AdjustTimeStamp(Time, TrackType);
-		});
+		MediaStreamSink->SetMediaSamplePoolAndQueue(VideoHardwareVideoDecodingSamplePool, &VideoSampleQueue,
+			[this](FTimespan Time, EMediaTrackType TrackType) -> FMediaTimeStamp
+			{
+				return AdjustTimeStamp(Time, TrackType);
+			},
+			[this]() -> bool
+			{
+				return IsSampleStillDesired();
+			}
+		);
 
 		if (FAILED(::MFCreateTopologyNode(MF_TOPOLOGY_OUTPUT_NODE, &OutputNode)) ||
 			FAILED(OutputNode->SetObject((IMFStreamSink*)MediaStreamSink)) ||
@@ -1969,7 +1997,13 @@ void FWmfMediaTracks::HandleMediaSamplerClock(EWmfMediaSamplerClockEvent Event,
 FMediaTimeStamp FWmfMediaTracks::AdjustTimeStamp(FTimespan Time, EMediaTrackType TrackType)
 {
 	FScopeLock Lock(&CriticalSection);
-	return FMediaTimeStamp(Time, FMediaTimeStamp::MakeSequenceIndex(SeekIndex, LoopIndex));
+	return FMediaTimeStamp(Time, SeekIndex, LoopIndex);
+}
+
+bool FWmfMediaTracks::IsSampleStillDesired()
+{
+	FScopeLock Lock(&CriticalSection);
+	return SeekIndex >= NextMinSequenceIndex.Get(0);
 }
 
 
@@ -2008,7 +2042,10 @@ void FWmfMediaTracks::HandleMediaSamplerAudioSample(const uint8* Buffer, uint32
 
 	if (AudioSample->Initialize(Buffer, Size, Format->Audio.NumChannels, Format->Audio.SampleRate, AdjustTimeStamp(Time, EMediaTrackType::Audio), Duration))
 	{
-		AudioSampleQueue.Enqueue(AudioSample);
+		if (AudioSample->GetTime().GetSequenceIndex() >= NextMinSequenceIndex.Get(0))
+		{
+			AudioSampleQueue.Enqueue(AudioSample);
+		}
 	}
 }
 
@@ -2038,7 +2075,10 @@ void FWmfMediaTracks::HandleMediaSamplerCaptionSample(const uint8* Buffer, uint3
 
 	if (CaptionSample->Initialize((char*)Buffer, AdjustTimeStamp(Time, EMediaTrackType::Caption), Duration))
 	{
-		CaptionSampleQueue.Enqueue(CaptionSample);
+		if (CaptionSample->GetTime().GetSequenceIndex() >= NextMinSequenceIndex.Get(0))
+		{
+			CaptionSampleQueue.Enqueue(CaptionSample);
+		}
 	}
 }
 
@@ -2068,7 +2108,10 @@ void FWmfMediaTracks::HandleMediaSamplerMetadataSample(const uint8* Buffer, uint
 
 	if (BinarySample->Initialize(Buffer, Size, AdjustTimeStamp(Time, EMediaTrackType::Metadata), Duration))
 	{
-		MetadataSampleQueue.Enqueue(BinarySample);
+		if (BinarySample->GetTime().GetSequenceIndex() >= NextMinSequenceIndex.Get(0))
+		{
+			MetadataSampleQueue.Enqueue(BinarySample);
+		}
 	}
 }
 
@@ -2131,7 +2174,10 @@ void FWmfMediaTracks::HandleMediaSamplerVideoSample(const uint8* Buffer, uint32
 		AdjustTimeStamp(Time, EMediaTrackType::Video),
 		Duration))
 	{
-		VideoSampleQueue.Enqueue(TextureSample);
+		if (TextureSample->GetTime().GetSequenceIndex() >= NextMinSequenceIndex.Get(0))
+		{
+			VideoSampleQueue.Enqueue(TextureSample);
+		}
 	}
 }
 
diff --git a/Engine/Plugins/Media/WmfMedia/Source/WmfMedia/Private/Player/WmfMediaTracks.h b/Engine/Plugins/Media/WmfMedia/Source/WmfMedia/Private/Player/WmfMediaTracks.h
--- a/Engine/Plugins/Media/WmfMedia/Source/WmfMedia/Private/Player/WmfMediaTracks.h
+++ b/Engine/Plugins/Media/WmfMedia/Source/WmfMedia/Private/Player/WmfMediaTracks.h
@@ -16,6 +16,7 @@
 #include "Math/IntPoint.h"
 #include "MediaSampleQueue.h"
 #include "Templates/SharedPointer.h"
+#include "Misc/Optional.h"
 
 #include "Windows/AllowWindowsPlatformTypes.h"
 
@@ -229,6 +230,7 @@ public:
 	virtual void FlushSamples() override;
 
 	virtual EFetchBestSampleResult FetchBestVideoSampleForTimeRange(const TRange<FMediaTimeStamp>& TimeRange, TSharedPtr<IMediaTextureSample, ESPMode::ThreadSafe>& OutSample, bool bReverse, bool bConsistentResult) override;
+	virtual void SetMinExpectedNextSequenceIndex(TOptional<int32> InNextSequenceIndex) override;
 
 	virtual bool DiscardVideoSamples(const TRange<FMediaTimeStamp>& TimeRange, bool bReverse) override;
 	virtual bool DiscardAudioSamples(const TRange<FMediaTimeStamp>& TimeRange, bool bReverse) override;
@@ -326,6 +328,9 @@ private:
 
 	/** Make sure sequence index is up to date and combine with given timestamp */
 	FMediaTimeStamp AdjustTimeStamp(FTimespan Time, EMediaTrackType TrackType);
+	
+	/** Check if the current sequence index to apply is still the one that is wanted. */
+	bool IsSampleStillDesired();
 
 private:
 	/** Audio sample object pool. */
@@ -407,6 +412,7 @@ private:
 	FWmfMediaSession* Session;
 
 	/** SequenceIndex tracking */
+	TOptional<int32> NextMinSequenceIndex;
 	int32 SeekIndex;
 	int32 LoopIndex;
 };
diff --git a/Engine/Plugins/Media/WmfMedia/Source/WmfMedia/Private/Wmf/WmfMediaSession.cpp b/Engine/Plugins/Media/WmfMedia/Source/WmfMedia/Private/Wmf/WmfMediaSession.cpp
--- a/Engine/Plugins/Media/WmfMedia/Source/WmfMedia/Private/Wmf/WmfMediaSession.cpp
+++ b/Engine/Plugins/Media/WmfMedia/Source/WmfMedia/Private/Wmf/WmfMediaSession.cpp
@@ -543,7 +543,7 @@ bool FWmfMediaSession::Seek(const FTimespan& Time, const FMediaSeekParams& InAdd
 	}
 
 	check(InAdditionalParams.NewSequenceIndex.IsSet());
-	UE_LOG(LogWmfMedia, Verbose, TEXT("Session %p: Seeking to %s with sequence index"), this, *Time.ToString(), InAdditionalParams.NewSequenceIndex.GetValue());
+	UE_LOG(LogWmfMedia, Verbose, TEXT("Session %p: Seeking to %s with sequence index %d"), this, *Time.ToString(), InAdditionalParams.NewSequenceIndex.GetValue());
 
 	// validate seek
 	if (!CanControl(EMediaControl::Seek))
diff --git a/Engine/Plugins/Media/WmfMedia/Source/WmfMedia/Private/Wmf/WmfMediaStreamSink.cpp b/Engine/Plugins/Media/WmfMedia/Source/WmfMedia/Private/Wmf/WmfMediaStreamSink.cpp
--- a/Engine/Plugins/Media/WmfMedia/Source/WmfMedia/Private/Wmf/WmfMediaStreamSink.cpp
+++ b/Engine/Plugins/Media/WmfMedia/Source/WmfMedia/Private/Wmf/WmfMediaStreamSink.cpp
@@ -824,9 +824,15 @@ void FWmfMediaStreamSink::CopyTextureAndEnqueueSample(IMFSample* pSample)
 
 	UE_LOG(LogWmfMedia, VeryVerbose, TEXT("StreamSink::CopyTextureAndEnqueueSample Queue Size: %d"), VideoSampleQueue->Num());
 
+	if (!IsSampleStillDesired())
+	{
+		UE_LOG(LogWmfMedia, VeryVerbose, TEXT("Sample is from previous sequence, dropping sample"));
+		return;
+	}
+
 	if (IsVideoSampleQueueFull())
 	{
-		UE_LOG(LogWmfMedia, VeryVerbose, TEXT("Queue is full, dropping samples"));
+		UE_LOG(LogWmfMedia, VeryVerbose, TEXT("Queue is full, dropping sample"));
 		return;
 	}
 
@@ -1076,12 +1082,14 @@ void FWmfMediaStreamSink::CopyTextureAndEnqueueSample(IMFSample* pSample)
 void FWmfMediaStreamSink::SetMediaSamplePoolAndQueue(
 	TSharedPtr<FWmfMediaHardwareVideoDecodingTextureSamplePool>& InVideoSamplePool,
 	TMediaSampleQueue<IMediaTextureSample>* InVideoSampleQueue,
-	TFunction<FMediaTimeStamp(FTimespan, EMediaTrackType)>&& InAdjustTimeStamp)
+	TFunction<FMediaTimeStamp(FTimespan, EMediaTrackType)>&& InAdjustTimeStamp,
+	TFunction<bool()>&& InIsSampleStillDesired)
 {
 	FScopeLock Lock(&CriticalSection);
 	VideoSamplePool = InVideoSamplePool;
 	VideoSampleQueue = InVideoSampleQueue;
 	AdjustTimeStamp = MoveTemp(InAdjustTimeStamp);
+	IsSampleStillDesired = MoveTemp(InIsSampleStillDesired);
 }
 
 
diff --git a/Engine/Plugins/Media/WmfMedia/Source/WmfMedia/Private/Wmf/WmfMediaStreamSink.h b/Engine/Plugins/Media/WmfMedia/Source/WmfMedia/Private/Wmf/WmfMediaStreamSink.h
--- a/Engine/Plugins/Media/WmfMedia/Source/WmfMedia/Private/Wmf/WmfMediaStreamSink.h
+++ b/Engine/Plugins/Media/WmfMedia/Source/WmfMedia/Private/Wmf/WmfMediaStreamSink.h
@@ -176,7 +176,7 @@ public:
 	 * @param InVideoSampleQueue Sample queue to fill
 	 * @param InSequenceIndex Pointer to current sequence index value
 	 */
-	void SetMediaSamplePoolAndQueue(TSharedPtr<FWmfMediaHardwareVideoDecodingTextureSamplePool>& InVideoSamplePool, TMediaSampleQueue<IMediaTextureSample>* InVideoSampleQueue, TFunction<FMediaTimeStamp(FTimespan, EMediaTrackType)>&& InAdjustTimeStamp);
+	void SetMediaSamplePoolAndQueue(TSharedPtr<FWmfMediaHardwareVideoDecodingTextureSamplePool>& InVideoSamplePool, TMediaSampleQueue<IMediaTextureSample>* InVideoSampleQueue, TFunction<FMediaTimeStamp(FTimespan, EMediaTrackType)>&& InAdjustTimeStamp, TFunction<bool()>&& InIsSampleStillDesired);
 
 	/**
 	 * Request more data being delivered to video output queue 
@@ -272,6 +272,9 @@ private:
 	/** Callback to adjust given timestamp for the output queue */
 	TFunction<FMediaTimeStamp(FTimespan,EMediaTrackType)> AdjustTimeStamp;
 
+	/** Callback to check if a newly decoded sample is still desired or should be discarded. */
+	TFunction<bool()> IsSampleStillDesired;
+
 	/** Make sure we don't spam output with MF_MT_SUBTYPE error message */
 	bool bShowSubTypeErrorMessage;
 
diff --git a/Engine/Plugins/Runtime/WebMMoviePlayer/Source/WebMMoviePlayer/Private/WebMMovieStreamer.cpp b/Engine/Plugins/Runtime/WebMMoviePlayer/Source/WebMMoviePlayer/Private/WebMMovieStreamer.cpp
--- a/Engine/Plugins/Runtime/WebMMoviePlayer/Source/WebMMoviePlayer/Private/WebMMovieStreamer.cpp
+++ b/Engine/Plugins/Runtime/WebMMoviePlayer/Source/WebMMoviePlayer/Private/WebMMovieStreamer.cpp
@@ -282,7 +282,7 @@ bool FWebMMovieStreamer::SendAudio(float InDeltaTime)
 		bFoundSample = Samples->FetchAudio(TimeRange, AudioSample);
 	}
 
-	return Samples->NumAudio() > 0 || AudioDecoder->IsBusy();
+	return Samples->NumAudioSamples() > 0 || AudioDecoder->IsBusy();
 }
 
 bool FWebMMovieStreamer::ReadMoreFrames()
diff --git a/Engine/Source/Runtime/Media/Public/IMediaPlayer.h b/Engine/Source/Runtime/Media/Public/IMediaPlayer.h
--- a/Engine/Source/Runtime/Media/Public/IMediaPlayer.h
+++ b/Engine/Source/Runtime/Media/Public/IMediaPlayer.h
@@ -331,15 +331,6 @@ public:
 		return true;
 	}
 
-	/**
-	 * Any extra processing that the player should do when FMediaPlayerFacade::ProcessVideoSamples
-	 * is run should be put here.
-	 */
-	virtual void ProcessVideoSamples()
-	{
-		// Override in child class if needed.
-	}
-
 	enum class EFeatureFlag
 	{
 		AllowShutdownOnClose = 0,		//!< Allow player to be shutdown right after 'close' event is received from it
diff --git a/Engine/Source/Runtime/Media/Public/IMediaSamples.h b/Engine/Source/Runtime/Media/Public/IMediaSamples.h
--- a/Engine/Source/Runtime/Media/Public/IMediaSamples.h
+++ b/Engine/Source/Runtime/Media/Public/IMediaSamples.h
@@ -5,6 +5,7 @@
 #include "Containers/Array.h"
 #include "Math/Range.h"
 #include "Misc/Timespan.h"
+#include "Misc/Optional.h"
 #include "Templates/SharedPointer.h"
 #include "IMediaTimeSource.h"
 
@@ -130,6 +131,7 @@ public:
 	{
 		return EFetchBestSampleResult::NotSupported;
 	}
+	virtual void SetMinExpectedNextSequenceIndex(TOptional<int32> InNextSequenceIndex) { }
 
 	virtual bool PeekVideoSampleTime(FMediaTimeStamp& TimeStamp) = 0;
 
@@ -142,9 +144,9 @@ public:
 	virtual bool DiscardSubtitleSamples(const TRange<FMediaTimeStamp>& TimeRange, bool bReverse) { return false; }
 	virtual bool DiscardMetadataSamples(const TRange<FMediaTimeStamp>& TimeRange, bool bReverse) { return false; }
 
-	virtual uint32 PurgeOutdatedVideoSamples(const FMediaTimeStamp & ReferenceTime, bool bReversed, FTimespan MaxAge) { return 0; };
+	virtual uint32 PurgeOutdatedVideoSamples(const FMediaTimeStamp& ReferenceTime, bool bReversed, FTimespan MaxAge) { return 0; };
 	virtual uint32 PurgeOutdatedCaptionSamples(const FMediaTimeStamp& ReferenceTime, bool bReversed, FTimespan MaxAge) { return 0; };
-	virtual uint32 PurgeOutdatedSubtitleSamples(const FMediaTimeStamp & ReferenceTime, bool bReversed, FTimespan MaxAge) { return 0; };
+	virtual uint32 PurgeOutdatedSubtitleSamples(const FMediaTimeStamp& ReferenceTime, bool bReversed, FTimespan MaxAge) { return 0; };
 	virtual uint32 PurgeOutdatedMetadataSamples(const FMediaTimeStamp& ReferenceTime, bool bReversed, FTimespan MaxAge) { return 0; };
 
 	virtual bool CanReceiveVideoSamples(uint32 Num) const { return true; }
@@ -153,8 +155,8 @@ public:
 	virtual bool CanReceiveCaptionSamples(uint32 Num) const { return true; }
 	virtual bool CanReceiveMetadataSamples(uint32 Num) const { return true; }
 
-	virtual int32 NumAudio() const { return -1; }
-	virtual int32 NumCaption() const { return -1; }
+	virtual int32 NumAudioSamples() const { return -1; }
+	virtual int32 NumCaptionSamples() const { return -1; }
 	virtual int32 NumMetadataSamples() const { return -1; }
 	virtual int32 NumSubtitleSamples() const { return -1; }
 	virtual int32 NumVideoSamples() const { return -1; }
diff --git a/Engine/Source/Runtime/Media/Public/IMediaTimeSource.h b/Engine/Source/Runtime/Media/Public/IMediaTimeSource.h
--- a/Engine/Source/Runtime/Media/Public/IMediaTimeSource.h
+++ b/Engine/Source/Runtime/Media/Public/IMediaTimeSource.h
@@ -27,94 +27,170 @@
 class FMediaTimeStamp
 {
 public:
-	FMediaTimeStamp() : Time(FTimespan::MinValue()), SequenceIndex(0) {}
-	explicit FMediaTimeStamp(const FTimespan & InTime) : Time(InTime), SequenceIndex(0) {}
-	explicit FMediaTimeStamp(const FTimespan & InTime, int64 InSequenceIndex) : Time(InTime), SequenceIndex(InSequenceIndex) {}
-
-	void Invalidate() { Time = FTimespan::MinValue(); }
-	bool IsValid() const { return Time != FTimespan::MinValue(); }
-
-	bool IsRelative() const { return SequenceIndex < 0; }
+	FMediaTimeStamp() : Time(FTimespan::MinValue())
+	{ }
+	explicit FMediaTimeStamp(const FTimespan& InTime) : Time(InTime)
+	{ }
+	explicit FMediaTimeStamp(const FTimespan& InTime, int64 InSequenceValue)
+		: Time(InTime)
+		, SequenceIndex(GetSequenceIndex(InSequenceValue))
+		, LoopIndex(GetLoopIndex(InSequenceValue))
+	{ }
+	explicit FMediaTimeStamp(const FTimespan& InTime, int32 InSequenceIndex, int32 InLoopIndex)
+		: Time(InTime)
+		, SequenceIndex(InSequenceIndex)
+		, LoopIndex(InLoopIndex)
+	{ }
+
+	void Invalidate()
+	{
+		Time = FTimespan::MinValue();
+	}
+	bool IsValid() const
+	{
+		return Time != FTimespan::MinValue();
+	}
 
-	bool operator < (const FMediaTimeStamp & Other) const { return (SequenceIndex < Other.SequenceIndex) || (SequenceIndex == Other.SequenceIndex && Time < Other.Time); }
-	bool operator <= (const FMediaTimeStamp & Other) const { return (SequenceIndex < Other.SequenceIndex) || (SequenceIndex == Other.SequenceIndex && Time <= Other.Time); }
-	bool operator == (const FMediaTimeStamp & Other) const { return (SequenceIndex == Other.SequenceIndex && Time == Other.Time); }
-	bool operator > (const FMediaTimeStamp & Other) const { return (SequenceIndex > Other.SequenceIndex) || (SequenceIndex == Other.SequenceIndex && Time > Other.Time); }
-	bool operator >= (const FMediaTimeStamp & Other) const { return (SequenceIndex > Other.SequenceIndex) || (SequenceIndex == Other.SequenceIndex && Time >= Other.Time); }
+	bool operator == (const FMediaTimeStamp& Other) const
+	{
+		return SequenceIndex == Other.SequenceIndex && LoopIndex == Other.LoopIndex && Time == Other.Time;
+	}
+	bool operator != (const FMediaTimeStamp& Other) const
+	{
+		return SequenceIndex != Other.SequenceIndex || LoopIndex != Other.LoopIndex || Time != Other.Time;
+	}
+	bool operator < (const FMediaTimeStamp& Other) const
+	{
+		return (GetSequenceIndex() < Other.GetSequenceIndex()) ||
+			   (GetSequenceIndex() == Other.GetSequenceIndex() && GetLoopIndex() < Other.GetLoopIndex()) ||
+			   (GetSequenceIndex() == Other.GetSequenceIndex() && GetLoopIndex() == Other.GetLoopIndex() && GetTime() < Other.GetTime());
+	}
+	bool operator <= (const FMediaTimeStamp& Other) const
+	{
+		return this->operator==(Other) || this->operator<(Other);
+	}
+	bool operator > (const FMediaTimeStamp& Other) const
+	{
+		return !this->operator<=(Other);
+	}
+	bool operator >= (const FMediaTimeStamp & Other) const
+	{
+		return !this->operator<(Other);
+	}
 
-	FMediaTimeStamp& SetPrimaryIndex(int32 InSetTo)
+	FMediaTimeStamp& SetTimeAndIndexValue(const FTimespan &InTime, int64 InIndexValue)
 	{
-		SequenceIndex = MakeSequenceIndex(InSetTo, GetSecondaryIndex());
+		Time = InTime;
+		SequenceIndex = GetSequenceIndex(InIndexValue);
+		LoopIndex = GetLoopIndex(InIndexValue);
 		return *this;
 	}
-
-	FMediaTimeStamp& SetSecondaryIndex(int32 InSetTo)
+	FMediaTimeStamp& SetTime(const FTimespan &InTime)
 	{
-		SequenceIndex = MakeSequenceIndex(GetPrimaryIndex(), InSetTo);
+		Time = InTime;
 		return *this;
 	}
 
-	FMediaTimeStamp& AdjustPrimaryIndex(int32 Add)
+	FMediaTimeStamp& SetSequenceIndex(int32 InSequenceIndex)
 	{
-		SequenceIndex += (static_cast<int64>(Add) << 32);
+		SequenceIndex = InSequenceIndex;
 		return *this;
 	}
-
-	FMediaTimeStamp& AdjustSecondaryIndex(int32 Add)
+	FMediaTimeStamp& SetLoopIndex(int32 InLoopIndex)
 	{
-		SequenceIndex += Add;
+		LoopIndex = InLoopIndex;
 		return *this;
 	}
 
-	int32 GetPrimaryIndex() const
+	const FTimespan& GetTime() const
 	{
-		return static_cast<int32>(SequenceIndex >> 32);
+		return Time;
 	}
-
-	int32 GetSecondaryIndex() const
+	int64 GetIndexValue() const
 	{
-		return static_cast<int32>(((uint64)SequenceIndex) & 0xffffffff);
+		return (int64)(((uint64)SequenceIndex << 32) | (uint32)LoopIndex);
 	}
 
+	FMediaTimeStamp& AdjustLoopIndex(int32 Add)
+	{
+		LoopIndex += Add;
+		return *this;
+	}
 
-	static int64 MakeSequenceIndex(int32 PrimaryIndex, int32 SecondaryIndex)
+	int32 GetSequenceIndex() const
 	{
-		return (static_cast<int64>(PrimaryIndex) << 32) + int64(SecondaryIndex);
+		return SequenceIndex;
 	}
 
-	static int64 AdjustPrimaryIndex(int64 InSequenceIndex, int32 Add)
+	int32 GetLoopIndex() const
 	{
-		return InSequenceIndex + (static_cast<int64>(Add) << 32);
+		return LoopIndex;
 	}
 
-	static int64 AdjustSecondaryIndex(int64 InSequenceIndex, int32 Add)
+
+	static int64 MakeIndexValue(int32 InSequenceIndex, int32 InLoopIndex)
 	{
-		return InSequenceIndex + Add;
+		return (int64)(((uint64)InSequenceIndex << 32) | (uint64)InLoopIndex);
 	}
 
-	static int32 GetPrimaryIndex(int64 InSequenceIndex)
+	static int32 GetSequenceIndex(int64 InSequenceIndex)
 	{
-		// note: needs to cope with signed operation
-		return static_cast<int32>((InSequenceIndex - int32(InSequenceIndex)) >> 32);	//-V1065 PSV warning "Expression can be simplified" suppressed
+		return static_cast<int32>(InSequenceIndex >> 32);
 	}
 
-	static int32 GetSecondaryIndex(int64 InSequenceIndex)
+	static int32 GetLoopIndex(int64 InSequenceIndex)
 	{
-		return static_cast<int32>(InSequenceIndex);
+		return static_cast<int32>(((uint64)InSequenceIndex) & 0xffffffff);
 	}
 
-	FMediaTimeStamp operator + (const FTimespan & Other) const { return FMediaTimeStamp(Time + Other, SequenceIndex); }
-	FMediaTimeStamp operator - (const FTimespan & Other) const { return FMediaTimeStamp(Time - Other, SequenceIndex); }
+	FMediaTimeStamp operator + (const FTimespan& Other) const
+	{
+		return FMediaTimeStamp(Time + Other, SequenceIndex, LoopIndex);
+	}
+	FMediaTimeStamp operator - (const FTimespan& Other) const
+	{
+		return FMediaTimeStamp(Time - Other, SequenceIndex, LoopIndex);
+	}
 
-	FMediaTimeStamp operator - (const FMediaTimeStamp& Other) const { return (Other.SequenceIndex < SequenceIndex) ? FMediaTimeStamp(FTimespan::MaxValue(), MAX_int64) : ((Other.SequenceIndex > SequenceIndex) ? FMediaTimeStamp(FTimespan::MinValue(), MAX_int64) : FMediaTimeStamp(Time - Other.Time, MAX_int64)); }
+	FMediaTimeStamp operator - (const FMediaTimeStamp& Other) const
+	{
+		if (Other.SequenceIndex == SequenceIndex && Other.LoopIndex == LoopIndex)
+		{
+			return FMediaTimeStamp(Time - Other.Time, MAX_int32, MAX_int32);
+		}
+		else if (Other.SequenceIndex <= SequenceIndex || Other.LoopIndex <= LoopIndex)
+		{
+			return FMediaTimeStamp(FTimespan::MaxValue(), MAX_int32, MAX_int32);
+		}
+		else
+		{
+			return FMediaTimeStamp(FTimespan::MinValue(), MAX_int32, MAX_int32);
+		}
+	}
 
-	FMediaTimeStamp& operator += (const FTimespan & Other) { Time += Other; return *this; }
-	FMediaTimeStamp& operator -= (const FTimespan & Other) { Time -= Other; return *this; }
+	FMediaTimeStamp& operator += (const FTimespan& Other)
+	{
+		Time += Other;
+		return *this;
+	}
+	FMediaTimeStamp& operator -= (const FTimespan& Other)
+	{
+		Time -= Other;
+		return *this;
+	}
 
-	FMediaTimeStamp& operator -= (const FMediaTimeStamp & Other) { Time -= Other.Time; SequenceIndex = MAX_int64; return *this; }
+	FMediaTimeStamp& operator -= (const FMediaTimeStamp& Other)
+	{
+		Time -= Other.Time;
+		SequenceIndex = MAX_int32;
+		LoopIndex = MAX_int32;
+		return *this;
+	}
 
 	FTimespan Time;
-	int64 SequenceIndex;
+private:
+	int32 SequenceIndex = 0;
+	int32 LoopIndex = 0;
 };
 
 class FMediaTimeStampSample
diff --git a/Engine/Source/Runtime/MediaAssets/Private/Assets/MediaPlayer.cpp b/Engine/Source/Runtime/MediaAssets/Private/Assets/MediaPlayer.cpp
--- a/Engine/Source/Runtime/MediaAssets/Private/Assets/MediaPlayer.cpp
+++ b/Engine/Source/Runtime/MediaAssets/Private/Assets/MediaPlayer.cpp
@@ -383,7 +383,7 @@ UMediaTimeStampInfo* UMediaPlayer::GetTimeStamp() const
 	{
 		FMediaTimeStamp TimeStamp = PlayerFacade->GetTimeStamp();
 		TimeStampInfo->Time = TimeStamp.Time;
-		TimeStampInfo->SequenceIndex = TimeStamp.SequenceIndex;
+		TimeStampInfo->SequenceIndex = TimeStamp.GetIndexValue();
 	}
 	return TimeStampInfo;
 }
@@ -395,7 +395,7 @@ UMediaTimeStampInfo* UMediaPlayer::GetDisplayTimeStamp() const
 	{
 		FMediaTimeStamp TimeStamp = PlayerFacade->GetDisplayTimeStamp();
 		TimeStampInfo->Time = TimeStamp.Time;
-		TimeStampInfo->SequenceIndex = TimeStamp.SequenceIndex;
+		TimeStampInfo->SequenceIndex = TimeStamp.GetIndexValue();
 	}
 	return TimeStampInfo;
 }
diff --git a/Engine/Source/Runtime/MediaUtils/Private/MediaAudioResampler.cpp b/Engine/Source/Runtime/MediaUtils/Private/MediaAudioResampler.cpp
--- a/Engine/Source/Runtime/MediaUtils/Private/MediaAudioResampler.cpp
+++ b/Engine/Source/Runtime/MediaUtils/Private/MediaAudioResampler.cpp
@@ -388,8 +388,7 @@ uint32 FMediaAudioResampler::Generate(float* Output, FMediaTimeStamp& OutTime, c
 
 		// calculate output time
 		const float InputOffset = (FrameIndex + FrameAlpha) / InputFrames;
-		OutTime.Time = InputTime.Time + InputOffset * InputDuration;
-		OutTime.SequenceIndex = InputTime.SequenceIndex;
+		OutTime.SetTimeAndIndexValue(InputTime.Time + InputOffset * InputDuration, InputTime.GetIndexValue());
 
 		// invert buffer if rate reversed
 		if ((Rate * LastRate) < 0.0f)
diff --git a/Engine/Source/Runtime/MediaUtils/Private/MediaPlayerFacade.cpp b/Engine/Source/Runtime/MediaUtils/Private/MediaPlayerFacade.cpp
--- a/Engine/Source/Runtime/MediaUtils/Private/MediaPlayerFacade.cpp
+++ b/Engine/Source/Runtime/MediaUtils/Private/MediaPlayerFacade.cpp
@@ -1140,7 +1140,11 @@ bool FMediaPlayerFacade::ContinueOpen(IMediaPlayerLifecycleManagerDelegate::ICon
 		NextEstVideoTimeAtFrameStart.Invalidate();
 		SeekTargetTime.Invalidate();
 		NextSeekTime.Reset();
-		SeekIndex = 0;
+		NextSequenceIndex.Reset();
+		if (Player->GetPlayerFeatureFlag(IMediaPlayer::EFeatureFlag::UsePlaybackTimingV2))
+		{
+			NextSequenceIndex = 0;
+		}
 	}
 
 	ResetTracks();
@@ -1216,22 +1220,17 @@ bool FMediaPlayerFacade::Seek(const FTimespan& InTime)
 	}
 
 	FMediaSeekParams SeekParams;
-	// v2 timing players are expected to use the new sequence index we set up.
-	int32 SequenceIndexNow = SeekIndex;
+	// v2 timing players are *required* to use the new sequence index we set up.
+	TOptional<int32> SequenceIndexNow = NextSequenceIndex;
 	if (CurrentPlayer->GetPlayerFeatureFlag(IMediaPlayer::EFeatureFlag::UsePlaybackTimingV2))
 	{
-		SeekIndex += (GetUnpausedRate() < 0.0f) ? -1 : 1;
-		SeekParams.NewSequenceIndex = SeekIndex;
-	}
-	else
-	{
-		SeekIndex = 0;
+		SeekParams.NewSequenceIndex = NextSequenceIndex = NextSequenceIndex.Get(0) + 1;
 	}
 	// Issue the seek.
 	if (!CurrentPlayer->GetControls().Seek(Time, SeekParams))
 	{
 		// If that failed restore the sequence index.
-		SeekIndex = SequenceIndexNow;
+		NextSequenceIndex = SequenceIndexNow;
 		return false;
 	}
 
@@ -1242,6 +1241,7 @@ bool FMediaPlayerFacade::Seek(const FTimespan& InTime)
 	{
 		// Yes. Flush only the facade side of the system as needed for seeks
 		// (the player is expected to flush its internal queues as needed itself)
+		PrepareSampleQueueForSequenceIndex();
 		Flush(CurrentPlayer->GetPlayerFeatureFlag(IMediaPlayer::EFeatureFlag::PlayerUsesInternalFlushOnSeek), true);
 	}
 	else
@@ -1253,8 +1253,7 @@ bool FMediaPlayerFacade::Seek(const FTimespan& InTime)
 		}
 	}
 
-	SeekTargetTime = FMediaTimeStamp(Time, FMediaTimeStamp::MakeSequenceIndex(SeekIndex, 0));
-
+	SeekTargetTime = FMediaTimeStamp(Time, NextSequenceIndex.Get(0), 0);
 	return true;
 }
 
@@ -1389,8 +1388,7 @@ const TRange<FMediaTimeStamp>& FMediaPlayerFacade::FBlockOnRange::GetRange() con
 	if (!CurrentPlayer->GetControls().IsLooping())
 	{
 		int32 LastLoopIdxS = LastTimeRange.IsEmpty() ? 0 : (int32)FMath::FloorToInt(LastTimeRange.GetLowerBoundValue().GetTotalSeconds() / Duration.GetTotalSeconds());
-		int64 SequenceIndex = FMediaTimeStamp::MakeSequenceIndex(OnBlockPrimaryIndex, OnBlockSecondaryIndexOffset + LastLoopIdxS);
-		BlockOnRange = TRange<FMediaTimeStamp>(FMediaTimeStamp(Start, SequenceIndex), FMediaTimeStamp(End, SequenceIndex));
+		BlockOnRange = TRange<FMediaTimeStamp>(FMediaTimeStamp(Start, OnBlockPrimaryIndex, OnBlockSecondaryIndexOffset + LastLoopIdxS), FMediaTimeStamp(End, OnBlockPrimaryIndex, OnBlockSecondaryIndexOffset + LastLoopIdxS));
 	}
 	else
 	{
@@ -1443,7 +1441,7 @@ const TRange<FMediaTimeStamp>& FMediaPlayerFacade::FBlockOnRange::GetRange() con
 					// Yes. Adjust the base secondary index value to guarantee a simple continuation of the secondary index values
 					int32 LastLoopIdxS = (int32)FMath::FloorToInt(LastTimeRange.GetLowerBoundValue().GetTotalSeconds() / Duration.GetTotalSeconds());
 					// -LoopIdxS -> Compensate for new start index
-					// +LastLoopIdxS + 1 -> Move index one beyond last
+					// +LastLoopIdxS - 1 -> Move index one before last
 					OnBlockSecondaryIndexOffset = -LoopIdxS + LastLoopIdxS - 1;
 				}
 			}
@@ -1463,9 +1461,7 @@ const TRange<FMediaTimeStamp>& FMediaPlayerFacade::FBlockOnRange::GetRange() con
 		}
 
 		// Assemble final blocking range
-		auto SeqIndexStart = FMediaTimeStamp::MakeSequenceIndex(OnBlockPrimaryIndex, OnBlockSecondaryIndexOffset + LoopIdxS);
-		auto SeqIndexEnd = FMediaTimeStamp::MakeSequenceIndex(OnBlockPrimaryIndex, OnBlockSecondaryIndexOffset + LoopIdxE);
-		BlockOnRange = TRange<FMediaTimeStamp>(FMediaTimeStamp(Start, SeqIndexStart), FMediaTimeStamp(End, SeqIndexEnd));
+		BlockOnRange = TRange<FMediaTimeStamp>(FMediaTimeStamp(Start, OnBlockPrimaryIndex, OnBlockSecondaryIndexOffset + LoopIdxS), FMediaTimeStamp(End, OnBlockPrimaryIndex, OnBlockSecondaryIndexOffset + LoopIdxE));
 		check(!BlockOnRange.IsEmpty());
 	}
 
@@ -1575,12 +1571,7 @@ bool FMediaPlayerFacade::SetRate(float Rate)
 
 bool FMediaPlayerFacade::SetNativeVolume(float Volume)
 {
-	if (!Player.IsValid())
-	{
-		return false;
-	}
-
-	return Player->SetNativeVolume(Volume);
+	return Player.IsValid() ? Player->SetNativeVolume(Volume) : false;
 }
 
 
@@ -1675,8 +1666,8 @@ bool FMediaPlayerFacade::BlockOnFetch() const
 			if (!Player->GetControls().IsLooping())
 			{
 				// Yes. Is the sample outside the media's range?
-				// (note: this assumes the media starts at time ZERO - this will not be the case at all times (e.g. life playback) -- for now we assume a player will flagged blocked playback as invalid in that case!)
-				if (BR.GetUpperBoundValue() < FMediaTimeStamp(FTimespan::Zero(), BR.GetUpperBoundValue().SequenceIndex) || Player->GetControls().GetDuration() <= BR.GetLowerBoundValue().Time)
+				// (note: this assumes the media starts at time ZERO - this will not be the case at all times (e.g. live playback) -- for now we assume a player will flagged blocked playback as invalid in that case!)
+				if (BR.GetUpperBoundValue().GetTime() < FTimespan::Zero() || Player->GetControls().GetDuration() <= BR.GetLowerBoundValue().GetTime())
 				{
 					return false;
 				}
@@ -1743,7 +1734,7 @@ void FMediaPlayerFacade::Flush(bool bExcludePlayer, bool bOnSeek)
 	LastAudioRenderedSampleTime.Invalidate();
 	if (bOnSeek)
 	{
-		BlockOnRange.OnSeek(SeekIndex);
+		BlockOnRange.OnSeek(NextSequenceIndex.GetValue());
 	}
 	else
 	{
@@ -1763,10 +1754,10 @@ void FMediaPlayerFacade::Flush(bool bExcludePlayer, bool bOnSeek)
 			if (!bExcludePlayer && !LastVideoSampleProcessedTimeRange.IsEmpty())
 			{
 				// Players will reset their sequence index related values, but keep the playback position. Adjust our record accordingly...
-				int32 LoopIdxS = FMediaTimeStamp::GetSecondaryIndex(LastVideoSampleProcessedTimeRange.GetLowerBoundValue().SequenceIndex);
-				int32 LoopIdxE = FMediaTimeStamp::GetSecondaryIndex(LastVideoSampleProcessedTimeRange.GetUpperBoundValue().SequenceIndex);
-				LastVideoSampleProcessedTimeRange.SetLowerBoundValue(FMediaTimeStamp(LastVideoSampleProcessedTimeRange.GetLowerBoundValue().Time, FMediaTimeStamp::MakeSequenceIndex(0, 0)));
-				LastVideoSampleProcessedTimeRange.SetUpperBoundValue(FMediaTimeStamp(LastVideoSampleProcessedTimeRange.GetUpperBoundValue().Time, FMediaTimeStamp::MakeSequenceIndex(0, LoopIdxE - LoopIdxS)));
+				int32 LoopIdxS = LastVideoSampleProcessedTimeRange.GetLowerBoundValue().GetLoopIndex();
+				int32 LoopIdxE = LastVideoSampleProcessedTimeRange.GetUpperBoundValue().GetLoopIndex();
+				LastVideoSampleProcessedTimeRange.SetLowerBoundValue(FMediaTimeStamp(LastVideoSampleProcessedTimeRange.GetLowerBoundValue().Time, 0, 0));
+				LastVideoSampleProcessedTimeRange.SetUpperBoundValue(FMediaTimeStamp(LastVideoSampleProcessedTimeRange.GetUpperBoundValue().Time, 0, LoopIdxE - LoopIdxS));
 			}
 		}
 
@@ -1776,7 +1767,6 @@ void FMediaPlayerFacade::Flush(bool bExcludePlayer, bool bOnSeek)
 		SeekTargetTime.Invalidate();
 	}
 
-	// V1 only
 	NextVideoSampleTime = FTimespan::MinValue();
 }
 
@@ -2205,12 +2195,12 @@ void FMediaPlayerFacade::TickInput(FTimespan DeltaTime, FTimespan Timecode)
 			EMediaEvent Event;
 			if (bIsBroadcastAllowed)
 			{
-				while (QueuedEventBroadcasts.Dequeue(Event))
+				while(QueuedEventBroadcasts.Dequeue(Event))
 				{
 					MediaEvent.Broadcast(Event);
 				}
 			}
-			while (QueuedEvents.Dequeue(Event))
+			while(QueuedEvents.Dequeue(Event))
 			{
 				ProcessEvent(Event, bIsBroadcastAllowed);
 			}
@@ -2241,7 +2231,7 @@ void FMediaPlayerFacade::TickInput(FTimespan DeltaTime, FTimespan Timecode)
 			IMediaSamples& Samples = Player->GetSamples();
 
 			double BlockingStart = FPlatformTime::Seconds();
-			while (1)
+			while(1)
 			{
 				ProcessCaptionSamples(Samples, TimeRange);
 				ProcessSubtitleSamples(Samples, TimeRange);
@@ -2267,7 +2257,7 @@ void FMediaPlayerFacade::TickInput(FTimespan DeltaTime, FTimespan Timecode)
 
 				// Process deferred events & check for events that break the block
 				bool bEventCancelsBlock = false;
-				while (QueuedEvents.Dequeue(Event))
+				while(QueuedEvents.Dequeue(Event))
 				{
 					if (Event == EMediaEvent::MediaClosed || Event == EMediaEvent::MediaOpenFailed)
 					{
@@ -2334,14 +2324,14 @@ void FMediaPlayerFacade::TickFetch(FTimespan DeltaTime, FTimespan Timecode)
 		bool bIsBroadcastAllowed = bAreEventsSafeForAnyThread || IsInGameThread();
 		if (bIsBroadcastAllowed)
 		{
-			while (QueuedEventBroadcasts.Dequeue(Event))
+			while(QueuedEventBroadcasts.Dequeue(Event))
 			{
 				MediaEvent.Broadcast(Event);
 			}
 		}
 
 		// process deferred events
-		while (QueuedEvents.Dequeue(Event))
+		while(QueuedEvents.Dequeue(Event))
 		{
 			ProcessEvent(Event, bIsBroadcastAllowed);
 		}
@@ -2360,7 +2350,7 @@ void FMediaPlayerFacade::TickFetch(FTimespan DeltaTime, FTimespan Timecode)
 		{
 			// process deferred events
 			EMediaEvent Event;
-			while (QueuedEvents.Dequeue(Event))
+			while(QueuedEvents.Dequeue(Event))
 			{
 				ProcessEvent(Event, true);
 			}
@@ -2394,7 +2384,7 @@ void FMediaPlayerFacade::TickFetch(FTimespan DeltaTime, FTimespan Timecode)
 		bool Blocked = false;
 		FDateTime BlockedTime;
 
-		while (true)
+		while(true)
 		{
 			ProcessCaptionSamplesV1(Samples, TimeRange);
 			ProcessSubtitleSamplesV1(Samples, TimeRange);
@@ -2480,11 +2470,10 @@ void FMediaPlayerFacade::TickTickable()
 
 	if (Player->GetPlayerFeatureFlag(IMediaPlayer::EFeatureFlag::UsePlaybackTimingV2))
 	{
-		const FMediaTimeStamp Time = GetTimeStamp();
-
-		auto TimeRange = TRange<FMediaTimeStamp>::Inclusive(FMediaTimeStamp(FTimespan::MinValue(), MIN_int32), Time + MediaPlayerFacade::MetadataPreroll);
-
 		ProcessAudioSamples(Samples, TRange<FMediaTimeStamp>());
+
+		const FMediaTimeStamp Time = GetTimeStamp();
+		auto TimeRange = TRange<FMediaTimeStamp>::Inclusive(FMediaTimeStamp(FTimespan::MinValue(), 0, 0), Time + MediaPlayerFacade::MetadataPreroll);
 		ProcessMetadataSamples(Samples, TimeRange);
 	}
 	else
@@ -2509,7 +2498,20 @@ void FMediaPlayerFacade::TickTickable()
 		ProcessMetadataSamplesV1(Samples, MetadataTimeRange);
 	}
 
-	SET_DWORD_STAT(STAT_MediaUtils_FacadeNumAudioSamples, Samples.NumAudio());
+	SET_DWORD_STAT(STAT_MediaUtils_FacadeNumAudioSamples, Samples.NumAudioSamples());
+}
+
+
+void FMediaPlayerFacade::PrepareSampleQueueForSequenceIndex()
+{
+	FScopeLock Lock(&CriticalSection);
+	if (!Player.IsValid() || !NextSequenceIndex.IsSet())
+	{
+		return;
+	}
+	int32 MinSeqIdx = NextSequenceIndex.GetValue();
+	IMediaSamples& Samples = Player->GetSamples();
+	Samples.SetMinExpectedNextSequenceIndex(NextSequenceIndex);
 }
 
 
@@ -2536,58 +2538,15 @@ void FMediaPlayerFacade::UpdateSeekStatus(const FMediaTimeStamp* pCheckTimeStamp
 
 			if (VideoTimeStamp.IsValid())
 			{
+				check(VideoTimeStamp.GetSequenceIndex() >= NextSequenceIndex.Get(0));
 				bool bRunningNonAudioClock = bHaveActiveAudio && !BlockOnRange.IsSet();
-
-				if (GetUnpausedRate() >= 0.0f)
-				{
-					// See if we already are looking at a sample from the target sequence index...
-					// (we are not checking for the precise location as some players might not able to deliver it)
-					if (FMediaTimeStamp::GetPrimaryIndex(VideoTimeStamp.SequenceIndex) < FMediaTimeStamp::GetPrimaryIndex(SeekTargetTime.SequenceIndex))
-					{
-						// No. Make sure we drop the sample & possible more up to the seek target (we use the fetch code to avoid any races with a old-sample-purge logic operating async)
-						Player->GetSamples().DiscardVideoSamples(TRange<FMediaTimeStamp>(VideoTimeStamp, SeekTargetTime), false);
-					}
-					else
-					{
-						// We have reached the sequence of the seek target, reset everything to normal operation...
-						// (we do not care if we reached the precise location beyond the index as we do not know if we even can)
-						if (bRunningNonAudioClock)
-						{
-							NextEstVideoTimeAtFrameStart = FMediaTimeStampSample(VideoTimeStamp, FPlatformTime::Seconds());
-						}
-
-						FScopeLock LockLT(&LastTimeValuesCS);
-
-						// Update the display version of then "current frame time" right now (to avoid any glitches as it can take a little while for the frame to actually change)
-						CurrentFrameVideoDisplayTimeStamp = SeekTargetTime;
-
-						// Seeking done
-						SeekTargetTime.Invalidate();
-
-					}
-				}
-				else
+				if (bRunningNonAudioClock)
 				{
-					// See if we already are looking at a sample from the target sequence index...
-					// (we are not checking for the precise location as some players might not able to deliver it)
-					if (FMediaTimeStamp::GetPrimaryIndex(VideoTimeStamp.SequenceIndex) > FMediaTimeStamp::GetPrimaryIndex(SeekTargetTime.SequenceIndex))
-					{
-						// No. Make sure we drop the sample & possible more up to the seek target (we use the fetch code to avoid any races with a old-sample-purge logic operating async)
-						Player->GetSamples().DiscardVideoSamples(TRange<FMediaTimeStamp>(VideoTimeStamp, SeekTargetTime), true);
-					}
-					else
-					{
-						// We have reached the sequence of the seek target, reset everything to normal operation...
-						// (we do not care if we reached the precise location beyond the index as we do not know if we even can)
-						if (bRunningNonAudioClock)
-						{
-							NextEstVideoTimeAtFrameStart = FMediaTimeStampSample(VideoTimeStamp, FPlatformTime::Seconds());
-						}
-
-						FScopeLock LockLT(&LastTimeValuesCS);
-						SeekTargetTime.Invalidate();
-					}
+					NextEstVideoTimeAtFrameStart = FMediaTimeStampSample(VideoTimeStamp, FPlatformTime::Seconds());
 				}
+				FScopeLock LockLT(&LastTimeValuesCS);
+				CurrentFrameVideoDisplayTimeStamp = SeekTargetTime;
+				SeekTargetTime.Invalidate();
 			}
 		}
 	}
@@ -2641,6 +2600,8 @@ void FMediaPlayerFacade::PreSampleProcessingTimeHandling()
 
 	FScopeLock Lock(&CriticalSection);
 
+	PrepareSampleQueueForSequenceIndex();
+
 	UpdateSeekStatus();
 
 	// No seeking?
@@ -2671,7 +2632,7 @@ void FMediaPlayerFacade::PreSampleProcessingTimeHandling()
 						// ...we need to see if the player's next sample might be so far in the future that we need to re-calibrate our timing
 						// (this could happen if the stream has a "gap" in PTS values - e.g. after pausing a live feed from a camera)
 						// (^^^ we do not do this on reverse playback as it is unlikely for such streams and might be thinned, hence show gaps under normal conditions)
-						if (VideoTimeStamp.SequenceIndex == NextEstVideoTimeAtFrameStart.TimeStamp.SequenceIndex)
+						if (VideoTimeStamp.GetIndexValue() == NextEstVideoTimeAtFrameStart.TimeStamp.GetIndexValue())
 						{
 							FTimespan Delta = VideoTimeStamp.Time - NextEstVideoTimeAtFrameStart.TimeStamp.Time;
 							if (GetUnpausedRate() < 0.0f)
@@ -2748,7 +2709,7 @@ void FMediaPlayerFacade::PostSampleProcessingTimeHandling(FTimespan DeltaTime)
 							while(NextEstVideoTimeAtFrameStart.TimeStamp.Time >= ActiveRangeEnd)
 							{
 								NextEstVideoTimeAtFrameStart.TimeStamp.Time -= ActiveRangeDuration;
-								NextEstVideoTimeAtFrameStart.TimeStamp.SequenceIndex = FMediaTimeStamp::AdjustSecondaryIndex(NextEstVideoTimeAtFrameStart.TimeStamp.SequenceIndex, 1);
+								NextEstVideoTimeAtFrameStart.TimeStamp.AdjustLoopIndex(1);
 							}
 						}
 						else
@@ -2756,7 +2717,7 @@ void FMediaPlayerFacade::PostSampleProcessingTimeHandling(FTimespan DeltaTime)
 							while(NextEstVideoTimeAtFrameStart.TimeStamp.Time < ActiveRangeStart)
 							{
 								NextEstVideoTimeAtFrameStart.TimeStamp.Time += ActiveRangeDuration;
-								NextEstVideoTimeAtFrameStart.TimeStamp.SequenceIndex = FMediaTimeStamp::AdjustSecondaryIndex(NextEstVideoTimeAtFrameStart.TimeStamp.SequenceIndex, -1);
+								NextEstVideoTimeAtFrameStart.TimeStamp.AdjustLoopIndex(-1);
 							}
 						}
 					}
@@ -2964,18 +2925,18 @@ bool FMediaPlayerFacade::GetCurrentPlaybackTimeRange(TRange<FMediaTimeStamp>& Ti
 			{
 				if (WrappedStart != TimeRange.GetLowerBoundValue().Time)
 				{
-					TimeRange.SetLowerBoundValue(FMediaTimeStamp(WrappedStart, FMediaTimeStamp::AdjustSecondaryIndex(TimeRange.GetLowerBoundValue().SequenceIndex, -1)));
+					TimeRange.SetLowerBoundValue(FMediaTimeStamp(WrappedStart, TimeRange.GetLowerBoundValue().GetSequenceIndex(), TimeRange.GetLowerBoundValue().GetLoopIndex() -1));
 				}
 				if (WrappedEnd != TimeRange.GetUpperBoundValue().Time)
 				{
-					TimeRange.SetUpperBoundValue(FMediaTimeStamp(WrappedEnd, FMediaTimeStamp::AdjustSecondaryIndex(TimeRange.GetUpperBoundValue().SequenceIndex, 1)));
+					TimeRange.SetUpperBoundValue(FMediaTimeStamp(WrappedEnd, TimeRange.GetUpperBoundValue().GetSequenceIndex(), TimeRange.GetUpperBoundValue().GetLoopIndex() + 1));
 				}
 			}
 		}
 		else
 		{
-			TimeRange.SetLowerBoundValue(FMediaTimeStamp(FMath::Clamp(TimeRange.GetLowerBoundValue().Time, ActiveRange.GetLowerBoundValue(), ActiveRange.GetUpperBoundValue()), TimeRange.GetLowerBoundValue().SequenceIndex));
-			TimeRange.SetUpperBoundValue(FMediaTimeStamp(FMath::Clamp(TimeRange.GetUpperBoundValue().Time, ActiveRange.GetLowerBoundValue(), ActiveRange.GetUpperBoundValue()), TimeRange.GetUpperBoundValue().SequenceIndex));
+			TimeRange.SetLowerBoundValue(FMediaTimeStamp(FMath::Clamp(TimeRange.GetLowerBoundValue().Time, ActiveRange.GetLowerBoundValue(), ActiveRange.GetUpperBoundValue()), TimeRange.GetLowerBoundValue().GetSequenceIndex(), TimeRange.GetLowerBoundValue().GetLoopIndex()));
+			TimeRange.SetUpperBoundValue(FMediaTimeStamp(FMath::Clamp(TimeRange.GetUpperBoundValue().Time, ActiveRange.GetLowerBoundValue(), ActiveRange.GetUpperBoundValue()), TimeRange.GetUpperBoundValue().GetSequenceIndex(), TimeRange.GetUpperBoundValue().GetLoopIndex()));
 		}
 	}
 
@@ -3003,26 +2964,18 @@ void FMediaPlayerFacade::ProcessAudioSamples(IMediaSamples& Samples, const TRang
 	// (we just have it around to be compatible / use older code that expects it)
 	check(TimeRange.GetLowerBound().IsOpen() && TimeRange.GetUpperBound().IsOpen());
 
-	// Seek in progress?
-	if (SeekTargetTime.IsValid())
-	{
-		// Yes. Fetch (and discard) all samples up to the seek target time...
-		// (we only throw out samples from prior sequence indices to make sure we do not swallow any audio from overlapping samples)
-		auto DiscardRange = TRange<FMediaTimeStamp>(FMediaTimeStamp(0), FMediaTimeStamp((CurrentRate >= 0.0f) ? FTimespan::Zero() : FTimespan::MaxValue(), SeekTargetTime.SequenceIndex));
-		Samples.DiscardAudioSamples(DiscardRange, GetUnpausedRate() < 0.0f);
-	}
-
 	//
 	// "Modern" 1-Audio-Sink-Only case (aka: we only feed the primary sink)
 	//
 	if (TSharedPtr< FMediaAudioSampleSink, ESPMode::ThreadSafe> PinnedPrimaryAudioSink = PrimaryAudioSink.Pin())
 	{
-		while (PinnedPrimaryAudioSink->CanAcceptSamples(1))
+		while(PinnedPrimaryAudioSink->CanAcceptSamples(1))
 		{
 			if (!Samples.FetchAudio(TimeRange, Sample))
+			{
 				break;
-
-			if (!Sample.IsValid())
+			}
+			else if (!Sample.IsValid())
 			{
 				continue;
 			}
@@ -3057,14 +3010,16 @@ void FMediaPlayerFacade::ProcessAudioSamples(IMediaSamples& Samples, const TRang
 				}
 
 			}
-			while (Samples.FetchAudio(TempRange, Sample))
-				;
+			while(Samples.FetchAudio(TempRange, Sample))
+			{
+			}
 		}
 		else
 		{
 			// No Video and no primary audio sink: we throw all away (sub-optimal as it will keep audio decoding busy; but this should be an edge case)
-			while (Samples.FetchAudio(TimeRange, Sample))
-				;
+			while(Samples.FetchAudio(TimeRange, Sample))
+			{
+			}
 		}
 	}
 }
@@ -3078,34 +3033,31 @@ bool FMediaPlayerFacade::IsVideoSampleStillGood(const TRange<FMediaTimeStamp>& L
 		// This is not the case: check more detailed!
 
 		// This better be true at all times
-		check(LastSampleTimeRange.GetLowerBoundValue().SequenceIndex == LastSampleTimeRange.GetUpperBoundValue().SequenceIndex);
+		check(LastSampleTimeRange.GetLowerBoundValue().GetIndexValue() == LastSampleTimeRange.GetUpperBoundValue().GetIndexValue());
 
 		// - - - - - - - - - - - - - - - - - - - - - - - - - - - - - - - - - - - - - - - - - - - - - - - - - - - - - - - - - - - - - - - - -
-		// Remap all values so we can assume all of them to be in a single "sequence index range" so the math doesn't get too unruly below
+		// Remap all values so we can assume all of them to be in a single "sequence index range" so the math doesn't get too unruly below.
+		// For that the sequence index must not have changed as we otherwise cannot assume any valid frame.
+		if (TimeRange.GetLowerBoundValue().GetSequenceIndex() != TimeRange.GetUpperBoundValue().GetSequenceIndex())
+		{
+			return false;
+		}
 
 		FTimespan Duration = Player->GetControls().GetDuration();
 
 		TRange<FMediaTimeStamp> TimeRange0;
-
-		// If we encounter a time range crossing some sequence index change, we need to check if we can "unroll" it...
-		uint64 LowerSeqIdx = TimeRange.GetLowerBoundValue().SequenceIndex;
-		uint64 UpperSeqIdx = TimeRange.GetUpperBoundValue().SequenceIndex;
-		uint64 RefSeqIdx = LowerSeqIdx;
-		if (LowerSeqIdx != UpperSeqIdx)
+		const int32 LoopIndex0 = TimeRange.GetLowerBoundValue().GetLoopIndex();
+		const int32 LoopIndex1 = TimeRange.GetUpperBoundValue().GetLoopIndex();
+		int32 RefLoopIndex = LoopIndex0;
+		// If we encounter a loop we need to check if we can "unroll" it...
+		if (LoopIndex0 != LoopIndex1)
 		{
-			if (FMediaTimeStamp::GetPrimaryIndex(LowerSeqIdx) != FMediaTimeStamp::GetPrimaryIndex(UpperSeqIdx))
-			{
-				// If we have a primary index change, we cannot assume any valid frame around...
-				return false;
-			}
-
 			// We only should get here with a looping player that knows its duration
 			check(Player->GetControls().IsLooping());
 			check(IsDurationValidAndFinite(Duration));
 
-			// So we must have a loop index change. Compute how many loops and change the range into one "unrolled" one as indicated by the playback direction...
-			int32 LoopIdxDiff = FMediaTimeStamp::GetSecondaryIndex(UpperSeqIdx)
-							  - FMediaTimeStamp::GetSecondaryIndex(LowerSeqIdx);
+			// Compute how many loops and change the range into one "unrolled" one as indicated by the playback direction...
+			int32 LoopIdxDiff = LoopIndex1 - LoopIndex0;
 			// Note: this will be positive even with reverse playback as the orientation of the range will no change
 			check(LoopIdxDiff > 0);
 
@@ -3118,13 +3070,13 @@ bool FMediaPlayerFacade::IsVideoSampleStillGood(const TRange<FMediaTimeStamp>& L
 			else
 			{
 				TimeRange0 = TRange<FMediaTimeStamp>(FMediaTimeStamp(TimeRange.GetLowerBoundValue().Time - FTimespan::FromSeconds(LoopIdxDiff * DurationD), 0), FMediaTimeStamp(TimeRange.GetUpperBoundValue().Time, 0));
-				RefSeqIdx = UpperSeqIdx;
+				RefLoopIndex = LoopIndex1;
 			}
 		}
 		else
 		{
 			// Simple case, just bring everything down to "zero sequence index" for ease of processing below...
-			TimeRange0 = TRange<FMediaTimeStamp>(FMediaTimeStamp(TimeRange.GetLowerBoundValue().Time, 0), FMediaTimeStamp(TimeRange.GetUpperBoundValue().Time, 0));
+			TimeRange0 = TRange<FMediaTimeStamp>(FMediaTimeStamp(TimeRange.GetLowerBoundValue().Time, 0, 0), FMediaTimeStamp(TimeRange.GetUpperBoundValue().Time, 0, 0));
 
 			// Is looping off?
 			if (!Player->GetControls().IsLooping())
@@ -3138,11 +3090,11 @@ bool FMediaPlayerFacade::IsVideoSampleStillGood(const TRange<FMediaTimeStamp>& L
 			}
 		}
 
-		// Map the last sample's time range to the same "sequence index range as the time range
-		// (note: for e.g. life streams that do not have any set duration all this will not change the timerange - just as needed)
-		int32 LastSampleIdxDiff = FMediaTimeStamp::GetSecondaryIndex(LastSampleTimeRange.GetLowerBoundValue().SequenceIndex) - RefSeqIdx;
-		FTimespan TimeOffset = IsDurationValidAndFinite(Duration) ? Duration * LastSampleIdxDiff : FTimespan::Zero();
-		TRange<FMediaTimeStamp> LastSampleTimeRange0(FMediaTimeStamp(LastSampleTimeRange.GetLowerBoundValue().Time + TimeOffset, 0), FMediaTimeStamp(LastSampleTimeRange.GetUpperBoundValue().Time + TimeOffset, 0));
+		// Map the last sample's time range to the same "sequence index" range as the time range
+		// (note: for Live streams that do not have any set duration all this will not change the timerange - just as needed)
+		const int32 LastSampleLoopDiff = LastSampleTimeRange.GetLowerBoundValue().GetLoopIndex() - RefLoopIndex;
+		FTimespan TimeOffset = IsDurationValidAndFinite(Duration) ? Duration * LastSampleLoopDiff : FTimespan::Zero();
+		TRange<FMediaTimeStamp> LastSampleTimeRange0(FMediaTimeStamp(LastSampleTimeRange.GetLowerBoundValue().Time + TimeOffset, 0, 0), FMediaTimeStamp(LastSampleTimeRange.GetUpperBoundValue().Time + TimeOffset, 0, 0));
 
 		// - - - - - - - - - - - - - - - - - - - - - - - - - - - - - - - - - - - - - - - - - - - - - - - - - - - - - - - - - - - -
 		// Now we can begin the checks with all time ranges mapped back to "zero sequence index"
@@ -3170,7 +3122,7 @@ bool FMediaPlayerFacade::IsVideoSampleStillGood(const TRange<FMediaTimeStamp>& L
 			//
 			// - We could check for them and generate proper changes to the sequence index
 			// - Doing this would leave us with quite complex setups to compute the coverage
-			// - We opt for a cleaner, simpler approach: as we are NOT interested into proper PTS values, we can safely work with an "infinite" time line when computing any overlaps, coverage and such
+			// - We opt for a cleaner, simpler approach: as we are NOT interested in proper PTS values here, we can safely work with an "infinite" time line when computing any overlaps, coverage and such
 			//   (note: we DO need to restrict the range to the actual media duration if not looping - the code above does this)
 			//
 			// --> we simply keep what we compute above!
@@ -3200,10 +3152,6 @@ bool FMediaPlayerFacade::ProcessVideoSamples(IMediaSamples& Samples, const TRang
 		return true;
 	}
 
-	// Let the player do some processing if needed.
-	// note: avoid using this - it will be deprecated
-	Player->ProcessVideoSamples();
-
 	// This is not to be used with V1 timing
 	check(Player->GetPlayerFeatureFlag(IMediaPlayer::EFeatureFlag::UsePlaybackTimingV2));
 	// We expect a fully closed range or we assume: nothing to do...
@@ -3259,7 +3207,7 @@ bool FMediaPlayerFacade::ProcessVideoSamples(IMediaSamples& Samples, const TRang
 				// (the FetchXYZ() code does not work well with a lower range limit at all - we ask for a "up to" type range instead
 				//  and limit the other side of the range in code here to not change the older logic & possibly cause trouble in old code)
 				TRange<FMediaTimeStamp> TempRange = bReverse ? TRange<FMediaTimeStamp>::AtLeast(TimeRange.GetUpperBoundValue()) : TRange<FMediaTimeStamp>::AtMost(TimeRange.GetUpperBoundValue());
-				while (Samples.FetchVideo(TempRange, Sample))
+				while(Samples.FetchVideo(TempRange, Sample))
 				{ }
 				if (Sample.IsValid() &&
 					((!bReverse && ((Sample->GetTime() + Sample->GetDuration()) > TimeRange.GetLowerBoundValue())) ||
@@ -3281,7 +3229,7 @@ bool FMediaPlayerFacade::ProcessVideoSamples(IMediaSamples& Samples, const TRang
 		// Use newest video frame available at all times (no Mediaframework timing control)
 		//
 		TRange<FMediaTimeStamp> TempRange; // fully open range
-		while (Samples.FetchVideo(TempRange, Sample))
+		while(Samples.FetchVideo(TempRange, Sample))
 		{ }
 	}
 
@@ -3324,12 +3272,12 @@ void FMediaPlayerFacade::ProcessCaptionSamples(IMediaSamples& Samples, const TRa
 	{
 		// Yes. Fetch (and discard) all samples up to the seek target time...
 		// (we only throw out samples from prior sequence indices to make sure we do not swallow any audio from overlapping samples)
-		auto DiscardRange = TRange<FMediaTimeStamp>(FMediaTimeStamp(0), FMediaTimeStamp((CurrentRate >= 0.0f) ? FTimespan::Zero() : FTimespan::MaxValue(), SeekTargetTime.SequenceIndex));
+		auto DiscardRange = TRange<FMediaTimeStamp>(FMediaTimeStamp(0), FMediaTimeStamp((CurrentRate >= 0.0f) ? FTimespan::Zero() : FTimespan::MaxValue(), SeekTargetTime.GetSequenceIndex(), SeekTargetTime.GetLoopIndex()));
 		Samples.DiscardCaptionSamples(DiscardRange, GetUnpausedRate() < 0.0f);
 	}
 	else
 	{
-		while (Samples.FetchCaption(TimeRange, Sample))
+		while(Samples.FetchCaption(TimeRange, Sample))
 		{
 			if (Sample.IsValid() && !CaptionSampleSinks.Enqueue(Sample.ToSharedRef()))
 			{
@@ -3353,12 +3301,12 @@ void FMediaPlayerFacade::ProcessSubtitleSamples(IMediaSamples& Samples, const TR
 	{
 		// Yes. Fetch (and discard) all samples up to the seek target time...
 		// (we only throw out samples from prior sequence indices to make sure we do not swallow any audio from overlapping samples)
-		auto DiscardRange = TRange<FMediaTimeStamp>(FMediaTimeStamp(0), FMediaTimeStamp((CurrentRate >= 0.0f) ? FTimespan::Zero() : FTimespan::MaxValue(), SeekTargetTime.SequenceIndex));
+		auto DiscardRange = TRange<FMediaTimeStamp>(FMediaTimeStamp(0), FMediaTimeStamp((CurrentRate >= 0.0f) ? FTimespan::Zero() : FTimespan::MaxValue(), SeekTargetTime.GetSequenceIndex(), SeekTargetTime.GetLoopIndex()));
 		Samples.DiscardSubtitleSamples(DiscardRange, GetUnpausedRate() < 0.0f);
 	}
 	else
 	{
-		while (Samples.FetchSubtitle(TimeRange, Sample))
+		while(Samples.FetchSubtitle(TimeRange, Sample))
 		{
 			//UE_LOG(LogMediaUtils, Display, TEXT("Subtitle @%.3f: %s"), Sample->GetTime().Time.GetTotalSeconds(), *Sample->GetText().ToString());
 			if (Sample.IsValid() && !SubtitleSampleSinks.Enqueue(Sample.ToSharedRef()))
@@ -3383,12 +3331,12 @@ void FMediaPlayerFacade::ProcessMetadataSamples(IMediaSamples& Samples, const TR
 	{
 		// Yes. Fetch (and discard) all samples up to the seek target time...
 		// (we only throw out samples from prior sequence indices to make sure we do not swallow any audio from overlapping samples)
-		auto DiscardRange = TRange<FMediaTimeStamp>(FMediaTimeStamp(0), FMediaTimeStamp((CurrentRate >= 0.0f) ? FTimespan::Zero() : FTimespan::MaxValue(), SeekTargetTime.SequenceIndex));
+		auto DiscardRange = TRange<FMediaTimeStamp>(FMediaTimeStamp(0), FMediaTimeStamp((CurrentRate >= 0.0f) ? FTimespan::Zero() : FTimespan::MaxValue(), SeekTargetTime.GetSequenceIndex(), SeekTargetTime.GetLoopIndex()));
 		Samples.DiscardMetadataSamples(DiscardRange, GetUnpausedRate() < 0.0f);
 	}
 	else
 	{
-		while (Samples.FetchMetadata(TimeRange, Sample))
+		while(Samples.FetchMetadata(TimeRange, Sample))
 		{
 			if (Sample.IsValid() && !MetadataSampleSinks.Enqueue(Sample.ToSharedRef()))
 			{
@@ -3405,7 +3353,7 @@ void FMediaPlayerFacade::ProcessAudioSamplesV1(IMediaSamples& Samples, TRange<FT
 {
 	TSharedPtr<IMediaAudioSample, ESPMode::ThreadSafe> Sample;
 
-	while (Samples.FetchAudio(TimeRange, Sample))
+	while(Samples.FetchAudio(TimeRange, Sample))
 	{
 		if (!Sample.IsValid())
 		{
@@ -3430,18 +3378,12 @@ void FMediaPlayerFacade::ProcessAudioSamplesV1(IMediaSamples& Samples, TRange<FT
 
 void FMediaPlayerFacade::ProcessVideoSamplesV1(IMediaSamples& Samples, TRange<FTimespan> TimeRange)
 {
-	// Let the player do some processing if needed.
-	if (Player.IsValid())
-	{
-		Player->ProcessVideoSamples();
-	}
-
 	// This is not to be used with V2 timing
 	check(!Player->GetPlayerFeatureFlag(IMediaPlayer::EFeatureFlag::UsePlaybackTimingV2));
 
 	TSharedPtr<IMediaTextureSample, ESPMode::ThreadSafe> Sample;
 
-	while (Samples.FetchVideo(TimeRange, Sample))
+	while(Samples.FetchVideo(TimeRange, Sample))
 	{
 		if (!Sample.IsValid())
 		{
@@ -3476,7 +3418,7 @@ void FMediaPlayerFacade::ProcessCaptionSamplesV1(IMediaSamples& Samples, TRange<
 {
 	TSharedPtr<IMediaOverlaySample, ESPMode::ThreadSafe> Sample;
 
-	while (Samples.FetchCaption(TimeRange, Sample))
+	while(Samples.FetchCaption(TimeRange, Sample))
 	{
 		if (Sample.IsValid() && !CaptionSampleSinks.Enqueue(Sample.ToSharedRef()))
 		{
@@ -3492,7 +3434,7 @@ void FMediaPlayerFacade::ProcessSubtitleSamplesV1(IMediaSamples& Samples, TRange
 {
 	TSharedPtr<IMediaOverlaySample, ESPMode::ThreadSafe> Sample;
 
-	while (Samples.FetchSubtitle(TimeRange, Sample))
+	while(Samples.FetchSubtitle(TimeRange, Sample))
 	{
 		if (Sample.IsValid() && !SubtitleSampleSinks.Enqueue(Sample.ToSharedRef()))
 		{
@@ -3511,7 +3453,7 @@ void FMediaPlayerFacade::ProcessMetadataSamplesV1(IMediaSamples& Samples, TRange
 {
 	TSharedPtr<IMediaBinarySample, ESPMode::ThreadSafe> Sample;
 
-	while (Samples.FetchMetadata(TimeRange, Sample))
+	while(Samples.FetchMetadata(TimeRange, Sample))
 	{
 		if (Sample.IsValid() && !MetadataSampleSinks.Enqueue(Sample.ToSharedRef()))
 		{
@@ -3534,10 +3476,18 @@ void FMediaPlayerFacade::ReceiveMediaEvent(EMediaEvent Event)
 		{
 			case EMediaEvent::Internal_PurgeVideoSamplesHint:
 			{
-				//
-				// Player asks to attempt to purge older samples in the video output queue it maintains
-				// (ask goes via facade as the player does not have accurate timing info)
-				//
+				/*
+					Sent by some media players to ask to purge older samples from the video output queue.
+					This is done to ensure that, even if the game thread is stalled and the facade is not
+					being ticked regularly where it would perform this task by passing frames from the
+					queue to the sink, frames that have passed the point where they should have been
+					sent to the sink will not clog the queue.
+					The player cannot perform this task on its own because it does not know the current
+					precise playback position.
+
+					Here we need to handle only everything not audio because audio is pulled by the
+					audio thread and not the gamethread, so it can never stall.
+				*/
 				TSharedPtr<IMediaPlayer, ESPMode::ThreadSafe> CurrentPlayer = Player;
 
 				if (!CurrentPlayer.IsValid())
@@ -3555,7 +3505,7 @@ void FMediaPlayerFacade::ReceiveMediaEvent(EMediaEvent Event)
 					return;
 				}
 
-				float Rate = CurrentRate;
+				const float Rate = CurrentRate;
 				if (Rate == 0.0f)
 				{
 					return;
@@ -3591,7 +3541,7 @@ void FMediaPlayerFacade::ReceiveMediaEvent(EMediaEvent Event)
 
 				// Captions...
 				NumPurged = 0;
-				if (CurrentPlayer->GetSamples().NumVideoSamples() >= kMinFramesInCaptionQueueToPurge)
+				if (CurrentPlayer->GetSamples().NumCaptionSamples() >= kMinFramesInCaptionQueueToPurge)
 				{
 					NumPurged = CurrentPlayer->GetSamples().PurgeOutdatedCaptionSamples(TimeRange.GetLowerBoundValue(), bReverse, FTimespan::FromSeconds(kOutdatedSubtitleSamplesTolerance * RateFactor));
 				}
@@ -3600,13 +3550,12 @@ void FMediaPlayerFacade::ReceiveMediaEvent(EMediaEvent Event)
 
 				// Subtitles...
 				NumPurged = 0;
-				if (CurrentPlayer->GetSamples().NumVideoSamples() >= kMinFramesInSubtitleQueueToPurge)
+				if (CurrentPlayer->GetSamples().NumSubtitleSamples() >= kMinFramesInSubtitleQueueToPurge)
 				{
 					NumPurged = CurrentPlayer->GetSamples().PurgeOutdatedSubtitleSamples(TimeRange.GetLowerBoundValue(), bReverse, FTimespan::FromSeconds(kOutdatedSubtitleSamplesTolerance * RateFactor));
 				}
 				SET_DWORD_STAT(STAT_MediaUtils_FacadeNumPurgedCaptionSamples, NumPurged);
 				INC_DWORD_STAT_BY(STAT_MediaUtils_FacadeTotalPurgedCaptionSamples, NumPurged);
-
 				break;
 			}
 
@@ -3677,7 +3626,7 @@ void FMediaPlayerFacade::ReInit()
 		MostRecentlyDeliveredVideoFrameTimecode.Reset();
 		SeekTargetTime.Invalidate();
 		NextSeekTime.Reset();
-		SeekIndex = 0;
+		NextSequenceIndex.Reset();
 	}
 }
 
diff --git a/Engine/Source/Runtime/MediaUtils/Private/MediaSamples.cpp b/Engine/Source/Runtime/MediaUtils/Private/MediaSamples.cpp
--- a/Engine/Source/Runtime/MediaUtils/Private/MediaSamples.cpp
+++ b/Engine/Source/Runtime/MediaUtils/Private/MediaSamples.cpp
@@ -137,6 +137,69 @@ void FMediaSamples::FlushSamples()
 	VideoSampleQueue.RequestFlush();
 }
 
+void FMediaSamples::SetMinExpectedNextSequenceIndex(TOptional<int32> InNextSequenceIndex)
+{
+	if (!InNextSequenceIndex.IsSet())
+	{
+		NextMinSequenceIndex.Reset();
+		return;
+	}
+	if (NextMinSequenceIndex == InNextSequenceIndex)
+	{
+		return;
+	}
+	NextMinSequenceIndex = InNextSequenceIndex;
+	int32 nsi = InNextSequenceIndex.GetValue();
+
+	AudioSampleQueue.PurgeUntilSequenceIndex(nsi);
+	MetadataSampleQueue.PurgeUntilSequenceIndex(nsi);
+	CaptionSampleQueue.PurgeUntilSequenceIndex(nsi);
+	SubtitleSampleQueue.PurgeUntilSequenceIndex(nsi);
+	VideoSampleQueue.PurgeUntilSequenceIndex(nsi);
+}
+
+void FMediaSamples::AddAudio(const TSharedRef<IMediaAudioSample, ESPMode::ThreadSafe>& Sample)
+{
+	if (Sample->GetTime().GetSequenceIndex() >= NextMinSequenceIndex.Get(0))
+	{
+		AudioSampleQueue.Enqueue(Sample);
+	}
+}
+
+void FMediaSamples::AddCaption(const TSharedRef<IMediaOverlaySample, ESPMode::ThreadSafe>& Sample)
+{
+	if (Sample->GetTime().GetSequenceIndex() >= NextMinSequenceIndex.Get(0))
+	{
+		CaptionSampleQueue.Enqueue(Sample);
+	}
+}
+
+void FMediaSamples::AddMetadata(const TSharedRef<IMediaBinarySample, ESPMode::ThreadSafe>& Sample)
+{
+	if (Sample->GetTime().GetSequenceIndex() >= NextMinSequenceIndex.Get(0))
+	{
+		MetadataSampleQueue.Enqueue(Sample);
+	}
+}
+
+void FMediaSamples::AddSubtitle(const TSharedRef<IMediaOverlaySample, ESPMode::ThreadSafe>& Sample)
+{
+	if (Sample->GetTime().GetSequenceIndex() >= NextMinSequenceIndex.Get(0))
+	{
+		SubtitleSampleQueue.Enqueue(Sample);
+	}
+}
+
+void FMediaSamples::AddVideo(const TSharedRef<IMediaTextureSample, ESPMode::ThreadSafe>& Sample)
+{
+	if (Sample->GetTime().GetSequenceIndex() >= NextMinSequenceIndex.Get(0))
+	{
+		VideoSampleQueue.Enqueue(Sample);
+	}
+}
+
+
+
 /**
  * Fetch video sample best suited for the given time range. Samples prior to the selected one will be removed from the queue.
  */
diff --git a/Engine/Source/Runtime/MediaUtils/Public/MediaPlayerFacade.h b/Engine/Source/Runtime/MediaUtils/Public/MediaPlayerFacade.h
--- a/Engine/Source/Runtime/MediaUtils/Public/MediaPlayerFacade.h
+++ b/Engine/Source/Runtime/MediaUtils/Public/MediaPlayerFacade.h
@@ -626,7 +626,7 @@ public:
 	 *	@remark This value is reset if Seek is called directly.
 	 */
 	MEDIAUTILS_API void SetNextSeek(const FTimespan& InTime);
-	
+
 	/**
 	 * Select the active track of the given type.
 	 *
@@ -936,6 +936,7 @@ private:
 	MEDIAUTILS_API void MonitorAudioEnablement();
 	MEDIAUTILS_API void UpdateSeekStatus(const FMediaTimeStamp* pCheckTimeStamp = nullptr);
 	void ExecuteNextSeek();
+	void PrepareSampleQueueForSequenceIndex();
 	MEDIAUTILS_API void PreSampleProcessingTimeHandling();
 	MEDIAUTILS_API bool GetCurrentPlaybackTimeRange(TRange<FMediaTimeStamp>& TimeRange, float Rate, FTimespan DeltaTime, bool bPurgeSampleRelated) const;
 	MEDIAUTILS_API void PostSampleProcessingTimeHandling(FTimespan DeltaTime);
@@ -1098,9 +1099,9 @@ private:
 
 	/** When multiple consecutive seek are issued while seeking, we keep track of the most recent request. */
 	TOptional<FTimespan> NextSeekTime;
-	
-	/** Current seek index */
-	int32 SeekIndex;
+
+	/** The next sequence index value we expect to see on the samples provided by a v2 timing enabled player. */
+	TOptional<int32> NextSequenceIndex;
 
 	/** Latch for error state of the most recently used player to be queried after it may have been closed. **/
 	bool bDidRecentPlayerHaveError;
diff --git a/Engine/Source/Runtime/MediaUtils/Public/MediaSampleQueue.h b/Engine/Source/Runtime/MediaUtils/Public/MediaSampleQueue.h
--- a/Engine/Source/Runtime/MediaUtils/Public/MediaSampleQueue.h
+++ b/Engine/Source/Runtime/MediaUtils/Public/MediaSampleQueue.h
@@ -111,29 +111,6 @@ public:
 		return true;
 	}
 
-	virtual bool PeekFrontAndBack(TSharedPtr<SampleType, ESPMode::ThreadSafe>& OutFirstSample, TSharedPtr<SampleType, ESPMode::ThreadSafe>& OutLastSample) override
-	{
-		FScopeLock Lock(&CriticalSection);
-
-		if (Samples.Num() == 0)
-		{
-			return false; // empty queue
-		}
-
-		TSharedPtr<SampleType, ESPMode::ThreadSafe> FirstSample(Samples[0]);
-		TSharedPtr<SampleType, ESPMode::ThreadSafe> LastSample(Samples.Last());
-
-		if (!FirstSample.IsValid() || !LastSample.IsValid())
-		{
-			return false; // pending flush
-		}
-
-		OutFirstSample = FirstSample;
-		OutLastSample = LastSample;
-
-		return true;
-	}
-
 	virtual void GetSampleTimes(TArray<TRange<FMediaTimeStamp>>& OutSampleTimeRanges) override
 	{
 		FScopeLock Lock(&CriticalSection);
@@ -342,6 +319,15 @@ public:
 		return 0;
 	}
 
+	void PurgeUntilSequenceIndex(int32 InUntilIndex)
+	{
+		FScopeLock Lock(&CriticalSection);
+		while(Samples.Num() && Samples[0]->GetTime().GetSequenceIndex() < InUntilIndex)
+		{
+			Samples.RemoveAt(0);
+		}
+	}
+
 public:
 
 	//~ TMediaSampleSink interface (to be called only from producer thread)
diff --git a/Engine/Source/Runtime/MediaUtils/Public/MediaSampleSource.h b/Engine/Source/Runtime/MediaUtils/Public/MediaSampleSource.h
--- a/Engine/Source/Runtime/MediaUtils/Public/MediaSampleSource.h
+++ b/Engine/Source/Runtime/MediaUtils/Public/MediaSampleSource.h
@@ -38,17 +38,6 @@ public:
 	 */
 	virtual bool Peek(TSharedPtr<SampleType, ESPMode::ThreadSafe>& OutSample) = 0;
 
-	/**
-	 * Simultaneously peeks at the next (the frontmost) and last samples in the queue without removing them.
-	 * The samples could be identical if there is only one element in the queue.
-	 *
-	 * @param OutFirstSample Will contain the frontmost sample.
-	 * @param OutLastSample Will contain the last sample, which could be identical to the frontmost one.
-	 * @return true if samples are returned, false otherwise.
-	 * @see Peek, Dequeue, Pop
-	 */
-	virtual bool PeekFrontAndBack(TSharedPtr<SampleType, ESPMode::ThreadSafe>& OutFirstSample, TSharedPtr<SampleType, ESPMode::ThreadSafe>& OutLastSample) = 0;
-
 	/**
 	 * Returns the sample start and end times of all samples currently in the queue.
 	 * @param OutSampleTimeRanges Will contain the sample time ranges.
diff --git a/Engine/Source/Runtime/MediaUtils/Public/MediaSamples.h b/Engine/Source/Runtime/MediaUtils/Public/MediaSamples.h
--- a/Engine/Source/Runtime/MediaUtils/Public/MediaSamples.h
+++ b/Engine/Source/Runtime/MediaUtils/Public/MediaSamples.h
@@ -44,10 +44,7 @@ public:
 	 * @param Sample The sample to add.
 	 * @see AddCaption, AddMetadata, AddSubtitle, AddVideo, NumAudio
 	 */
-	void AddAudio(const TSharedRef<IMediaAudioSample, ESPMode::ThreadSafe>& Sample)
-	{
-		AudioSampleQueue.Enqueue(Sample);
-	}
+	MEDIAUTILS_API void AddAudio(const TSharedRef<IMediaAudioSample, ESPMode::ThreadSafe>& Sample);
 
 	/**
 	 * Add the given caption sample to the cache.
@@ -55,10 +52,7 @@ public:
 	 * @param Sample The sample to add.
 	 * @see AddAudio, AddMetadata, AddSubtitle, AddVideo, NumCaption
 	 */
-	void AddCaption(const TSharedRef<IMediaOverlaySample, ESPMode::ThreadSafe>& Sample)
-	{
-		CaptionSampleQueue.Enqueue(Sample);
-	}
+	MEDIAUTILS_API void AddCaption(const TSharedRef<IMediaOverlaySample, ESPMode::ThreadSafe>& Sample);
 
 	/**
 	 * Add the given audio sample to the cache.
@@ -66,10 +60,7 @@ public:
 	 * @param Sample The sample to add.
 	 * @see AddAudio, AddCaption, AddSubtitle, AddVideo, NumMetadata
 	 */
-	void AddMetadata(const TSharedRef<IMediaBinarySample, ESPMode::ThreadSafe>& Sample)
-	{
-		MetadataSampleQueue.Enqueue(Sample);
-	}
+	MEDIAUTILS_API void AddMetadata(const TSharedRef<IMediaBinarySample, ESPMode::ThreadSafe>& Sample);
 
 	/**
 	 * Add the given subtitle sample to the cache.
@@ -77,10 +68,7 @@ public:
 	 * @param Sample The sample to add.
 	 * @see AddAudio, AddCaption, AddMetadata, AddVideo, NumSubtitle
 	 */
-	void AddSubtitle(const TSharedRef<IMediaOverlaySample, ESPMode::ThreadSafe>& Sample)
-	{
-		SubtitleSampleQueue.Enqueue(Sample);
-	}
+	MEDIAUTILS_API void AddSubtitle(const TSharedRef<IMediaOverlaySample, ESPMode::ThreadSafe>& Sample);
 
 	/**
 	 * Add the given audio sample to the cache.
@@ -88,18 +76,14 @@ public:
 	 * @param Sample The sample to add.
 	 * @see AddAudio, AddCaption, AddMetadata, AddSubtitle, NumVideo
 	 */
-	void AddVideo(const TSharedRef<IMediaTextureSample, ESPMode::ThreadSafe>& Sample)
-	{
-		VideoSampleQueue.Enqueue(Sample);
-	}
-
+	MEDIAUTILS_API void AddVideo(const TSharedRef<IMediaTextureSample, ESPMode::ThreadSafe>& Sample);
 	/**
 	 * Get the number of queued audio samples.
 	 *
 	 * @return Number of samples.
 	 * @see AddAudio, NumCaption, NumMetadata, NumSubtitle, NumVideo
 	 */
-	int32 NumAudio() const override
+	int32 NumAudioSamples() const override
 	{
 		return AudioSampleQueue.Num();
 	}
@@ -110,7 +94,7 @@ public:
 	 * @return Number of samples.
 	 * @see AddCaption, NumAudio, NumMetadata, NumSubtitle, NumVideo
 	 */
-	int32 NumCaption() const override
+	int32 NumCaptionSamples() const override
 	{
 		return CaptionSampleQueue.Num();
 	}
@@ -167,6 +151,9 @@ public:
 
 	MEDIAUTILS_API virtual EFetchBestSampleResult FetchBestVideoSampleForTimeRange(const TRange<FMediaTimeStamp>& TimeRange, TSharedPtr<IMediaTextureSample, ESPMode::ThreadSafe>& OutSample, bool bReverse, bool bConsistentResult) override;
 
+	MEDIAUTILS_API virtual void SetMinExpectedNextSequenceIndex(TOptional<int32> InNextSequenceIndex) override;
+
+
 	/**
 	 * Peek next video sample's timestamp
 	 * @return true if value could be retrieved, false otherwise
@@ -245,4 +232,7 @@ private:
 
 	/** Video sample queue. */
 	FMediaTextureSampleQueue VideoSampleQueue;
+
+	/** The next smallest expected sequence index if set. */
+	TOptional<int32> NextMinSequenceIndex;
 };
