Add cvar to force FMediaPlayerFacade::ContinueOpen(...) to fail for testing purposes.
- media players can add significant GPU overhead which makes it harder to profile GPU workloads.

diff --git a/Engine/Source/Runtime/MediaUtils/Private/MediaPlayerFacade.cpp b/Engine/Source/Runtime/MediaUtils/Private/MediaPlayerFacade.cpp
--- a/Engine/Source/Runtime/MediaUtils/Private/MediaPlayerFacade.cpp
+++ b/Engine/Source/Runtime/MediaUtils/Private/MediaPlayerFacade.cpp
@@ -91,6 +91,25 @@ static const int32	kMinFramesInVideoQueueToPurge = 3;			// we only consider purg
 static const int32	kMinFramesInSubtitleQueueToPurge = 3;		// we only consider purging any old frames from the subtitle queue if more than these are present (to not kill a slow playback entirely)
 static const int32	kMinFramesInCaptionQueueToPurge = 3;		// we only consider purging any old frames from the caption queue if more than these are present (to not kill a slow playback entirely)
 
+/* Cvars
+*****************************************************************************/
+
+namespace UE::MediaUtils::Private
+{
+
+#if !UE_BUILD_SHIPPING
+
+TAutoConsoleVariable<bool> CVarTestForcePlayerCreateFailed(
+	TEXT("m.Test.ForcePlayerCreateFailed"),
+	false,
+	TEXT("Whether force media player creation to fail."),
+	ECVF_ReadOnly
+);
+
+#endif // !UE_BUILD_SHIPPING
+
+}
+
 /* Local helpers
 *****************************************************************************/
 
@@ -1094,7 +1113,13 @@ bool FMediaPlayerFacade::ContinueOpen(IMediaPlayerLifecycleManagerDelegate::ICon
 		LifecycleManagerDelegateControl = NewLifecycleManagerDelegateControl;
 	}
 
-	if (!Player.IsValid())
+	bool bIsRequestInvalid = (Player == nullptr);
+
+#if !UE_BUILD_SHIPPING
+	bIsRequestInvalid = bIsRequestInvalid || UE::MediaUtils::Private::CVarTestForcePlayerCreateFailed.GetValueOnAnyThread();
+#endif
+
+	if (bIsRequestInvalid)
 	{
 		NotifyLifetimeManagerDelegate_PlayerCreateFailed();
 		// Make sure we don't get called from the "tickable" thread anymore - no need as we have no player
