MediaUtils: Changed the blocking range GetRange() code to use precise integer math in calculating the loop indices instead of floating point math with potential rounding issues. Renamed a few members to be hopefully more expressive of what it is they are tracking.

#rnx

diff --git a/Engine/Source/Runtime/MediaUtils/Private/MediaPlayerFacade.cpp b/Engine/Source/Runtime/MediaUtils/Private/MediaPlayerFacade.cpp
--- a/Engine/Source/Runtime/MediaUtils/Private/MediaPlayerFacade.cpp
+++ b/Engine/Source/Runtime/MediaUtils/Private/MediaPlayerFacade.cpp
@@ -1340,18 +1340,18 @@ void FMediaPlayerFacade::SetBlockOnTimeRange(const TRange<FTimespan>& TimeRange)
 
 void FMediaPlayerFacade::FBlockOnRange::OnFlush()
 {
-	LastTimeRange = TRange<FTimespan>::Empty();
-	OnBlockPrimaryIndex = 0;
-	OnBlockSecondaryIndexOffset = 0;
+	LastProcessedTimeRange = TRange<FTimespan>::Empty();
+	OnBlockSequenceIndex = 0;
+	OnBlockLoopIndexOffset = 0;
 	RangeIsDirty = true;
 }
 
 
 void FMediaPlayerFacade::FBlockOnRange::OnSeek(int32 PrimaryIndex)
 {
-	LastTimeRange = TRange<FTimespan>::Empty();
-	OnBlockPrimaryIndex = PrimaryIndex;
-	OnBlockSecondaryIndexOffset = 0;
+	LastProcessedTimeRange = TRange<FTimespan>::Empty();
+	OnBlockSequenceIndex = PrimaryIndex;
+	OnBlockLoopIndexOffset = 0;
 	RangeIsDirty = true;
 }
 
@@ -1392,7 +1392,7 @@ const TRange<FMediaTimeStamp>& FMediaPlayerFacade::FBlockOnRange::GetRange() con
 	// If the range is empty or the player can't support blocked playback: reset everything & return empty block range...
 	if (CurrentTimeRange.IsEmpty() || !CurrentPlayer->GetControls().CanControl(EMediaControl::BlockOnFetch))
 	{
-		LastTimeRange = TRange<FTimespan>::Empty();
+		LastProcessedTimeRange = TRange<FTimespan>::Empty();
 		BlockOnRange = TRange<FMediaTimeStamp>::Empty();
 		CurrentPlayer->GetControls().SetBlockingPlaybackHint(false);
 		return BlockOnRange;
@@ -1401,41 +1401,43 @@ const TRange<FMediaTimeStamp>& FMediaPlayerFacade::FBlockOnRange::GetRange() con
 	EMediaState PlayerState = CurrentPlayer->GetControls().GetState();
 	if (PlayerState != EMediaState::Paused && PlayerState != EMediaState::Playing)
 	{
-		// Return an empty range. Note that the "isSet()" method will still report a set block - so all code will remain in "external clock" mode,
+		// Return an empty range. Note that the "IsSet()" method will still report a set block - so all code will remain in "external clock" mode,
 		// but no samples will be requested (and no actual blocking should take place)
 		static auto EmptyRange(TRange<FMediaTimeStamp>::Empty());
 		return EmptyRange;
 	}
 
-	FTimespan Duration(CurrentPlayer->GetControls().GetDuration());
-	FTimespan Start(CurrentTimeRange.GetLowerBoundValue());
-	FTimespan End(CurrentTimeRange.GetUpperBoundValue());
+	const int64 StartTicks = CurrentTimeRange.GetLowerBoundValue().GetTicks();
+	const int64 EndExclusiveTicks = CurrentTimeRange.GetUpperBoundValue().GetTicks();
+	const int64 DurationTicks = CurrentPlayer->GetControls().GetDuration().GetTicks();
+	check(StartTicks >= 0 && EndExclusiveTicks >= StartTicks);
 
 	/*
-	* On the synthesized sequence and loop index values:
-	* - We track seeks and hence can insert the proper seek index easily, although the user does not provide it / does not need to track it
-	* - The loop index gets somewhat of a special treatment:
-	*  -- With tools like Sequencer a blocked range my speed along the time line quite quickly and if the player is configured as looping it would be expected that the video loops while this is done
-	*  -- We hence could loop multiple times within a single update interval
-	*  -- Still we treat any loop (aka: a jump "backwards" without an explicit seek) as a single loop iteration
-	*     (this is easier for any player to work with and provides the same visual results)
+		When looping we need to synthesize the expected start and end loop indices of the range we are returning.
+		This is because on playback start/seeking the media player implicitly starts with a loop index of 0,
+		while we could be in any of the looping repetitions in the sequencer track of this movie (if the track
+		has been pulled out "to the right" to have it repeat the clip n times).
+		Because of that we need to "lock" an initial loop offset that represents this initial difference and
+		needs to be adjusted with any looping of the sequencer.
+		Please note that this does not include the case where the media track ends before the end of the
+		sequencer. In that case the media player is closed and re-opened on the media track boundaries.
 	*/
-
 	if (!CurrentPlayer->GetControls().IsLooping())
 	{
-		int32 LastLoopIdxS = LastTimeRange.IsEmpty() ? 0 : (int32)FMath::FloorToInt(LastTimeRange.GetLowerBoundValue().GetTotalSeconds() / Duration.GetTotalSeconds());
-		BlockOnRange = TRange<FMediaTimeStamp>(FMediaTimeStamp(Start, OnBlockPrimaryIndex, OnBlockSecondaryIndexOffset + LastLoopIdxS), FMediaTimeStamp(End, OnBlockPrimaryIndex, OnBlockSecondaryIndexOffset + LastLoopIdxS));
+		FTimespan Start(CurrentTimeRange.GetLowerBoundValue());
+		FTimespan End(CurrentTimeRange.GetUpperBoundValue());
+		const int32 LoopIdx = LastProcessedTimeRange.IsEmpty() ? 0 : LastProcessedTimeRange.GetLowerBoundValue().GetTicks() / DurationTicks;
+		BlockOnRange = TRange<FMediaTimeStamp>(FMediaTimeStamp(Start, OnBlockSequenceIndex, OnBlockLoopIndexOffset + LoopIdx), FMediaTimeStamp(End, OnBlockSequenceIndex, OnBlockLoopIndexOffset + LoopIdx));
 	}
 	else
 	{
-		// If this would be called very early in the player's startup after open() we would not yet be known... that would be fatal
 		/*
-		 Should this actually happen in real-life applications, we could move the computations here into an accessor method used internally, so that this would be done
-		 only if data is processed, which would also mean: we know the duration!
-		 (Exception: live playback! --> but we would not allow blocking there anyway! (makes no sense as real life use case))
+			If this would be called very early in the player's startup after open() we would not yet be known... that would be fatal
+				Should this actually happen in real-life applications, we could move the computations here into an accessor method used internally, so that this would be done
+				only if data is processed, which would also mean: we know the duration!
+				(Exception: live playback! --> but we would not allow blocking there anyway! (makes no sense as real life use case))
 		*/
-		check(IsDurationValidAndFinite(Duration));
-		if (!IsDurationValidAndFinite(Duration))
+		if (!IsDurationValidAndFinite(CurrentPlayer->GetControls().GetDuration()))
 		{
 			// Catch if this is called to early and reset blocking...
 			BlockOnRange = TRange<FMediaTimeStamp>::Empty();
@@ -1443,61 +1445,55 @@ const TRange<FMediaTimeStamp>& FMediaPlayerFacade::FBlockOnRange::GetRange() con
 			return BlockOnRange;
 		}
 
-
-		bool bReverse = (Facade->GetUnpausedRate() < 0.0f);
-
-		// Modulo on the time to get it into media's range
-		// (assumes zero-start-time)
-		Start = WrappedModulo(Start, Duration);
-		End = WrappedModulo(End, Duration);
-
-		int32 LoopIdxS = (int32)FMath::FloorToInt(CurrentTimeRange.GetLowerBoundValue().GetTotalSeconds() / Duration.GetTotalSeconds());
-		int32 LoopIdxE = (int32)FMath::FloorToInt(CurrentTimeRange.GetUpperBoundValue().GetTotalSeconds() / Duration.GetTotalSeconds());
-
-		if (!LastTimeRange.IsEmpty())
+		FMediaTimeStamp t0, t1;
+		t0.SetTime(FTimespan(StartTicks % DurationTicks)).SetSequenceIndex(OnBlockSequenceIndex).SetLoopIndex(StartTicks / DurationTicks);
+		t1.SetTime(FTimespan(EndExclusiveTicks% DurationTicks)).SetSequenceIndex(OnBlockSequenceIndex).SetLoopIndex(EndExclusiveTicks / DurationTicks);
+		const bool bReverse = (Facade->GetUnpausedRate() < 0.0f);
+		// Did we process a time range before?
+		if (LastProcessedTimeRange.IsEmpty())
 		{
-			// Adjust loop index base such that, given a playback direction, we can guarantee that the new indices returned fit a monotone progression
+			/*
+				No, playback has just started fresh or through a seek.
+				The media player will start with a loop index of 0, but the blocking range could be anywhere within
+				a movie repetition (ie when the movie has been pulled out in the sequencer track to repeat a number of times).
+				We set that repetion count as the base for the loop index.
+			*/
+			check(OnBlockLoopIndexOffset == 0);
+			OnBlockLoopIndexOffset = !bReverse ? -t0.GetLoopIndex() : -t1.GetLoopIndex();
+		}
+		else
+		{
+			/*
+				We already processed a time range. We now need to check if the current one has wrapped around in the
+				current playback direction.
+
+				Theoretically, with either very, very short movies or an excessively huge delta time this could have
+				wrapped around more than once. We cannot detect this and hope this will not occur.
+			*/
 			if (!bReverse)
 			{
-				// Did the range passed in, loop around?
-				if (LastTimeRange.GetLowerBoundValue() > CurrentTimeRange.GetLowerBoundValue())
+				if (LastProcessedTimeRange.GetLowerBoundValue() > CurrentTimeRange.GetLowerBoundValue())
 				{
-					// Yes. Adjust the base secondary index value to guarantee a simple continuation of the secondary index values
-					int32 LastLoopIdxS = (int32)FMath::FloorToInt(LastTimeRange.GetLowerBoundValue().GetTotalSeconds() / Duration.GetTotalSeconds());
-					// -LoopIdxS -> Compensate for new start index
-					// +LastLoopIdxS + 1 -> Move index one beyond last
-					OnBlockSecondaryIndexOffset = -LoopIdxS + LastLoopIdxS + 1;
+					// Figure the loop index of the last range's start time.
+					const int32 LastRangeLoopIdx = LastProcessedTimeRange.GetLowerBoundValue().GetTicks() / DurationTicks;
+					OnBlockLoopIndexOffset += LastRangeLoopIdx + 1;
 				}
 			}
 			else
 			{
-				// Did the range passed in, loop around?
-				if (LastTimeRange.GetLowerBoundValue() < CurrentTimeRange.GetLowerBoundValue())
+				if (LastProcessedTimeRange.GetLowerBoundValue() < CurrentTimeRange.GetLowerBoundValue())
 				{
-					// Yes. Adjust the base secondary index value to guarantee a simple continuation of the secondary index values
-					int32 LastLoopIdxS = (int32)FMath::FloorToInt(LastTimeRange.GetLowerBoundValue().GetTotalSeconds() / Duration.GetTotalSeconds());
-					// -LoopIdxS -> Compensate for new start index
-					// +LastLoopIdxS - 1 -> Move index one before last
-					OnBlockSecondaryIndexOffset = -LoopIdxS + LastLoopIdxS - 1;
+					// Figure the loop index of the last range's start time.
+					const int32 ThisRangeLoopIdx = t0.GetLoopIndex();
+					OnBlockLoopIndexOffset -= ThisRangeLoopIdx + 1;
 				}
 			}
 		}
-		else
-		{
-			// No old range data. We must assume this as the first block after startup / flush and start at loop index zero. Relocate indices by moving the base, so we really start at zero...
-			check(OnBlockSecondaryIndexOffset == 0);
-			if (!bReverse)
-			{
-				OnBlockSecondaryIndexOffset = -LoopIdxS;
-			}
-			else
-			{
-				OnBlockSecondaryIndexOffset = -LoopIdxE;
-			}
-		}
 
 		// Assemble final blocking range
-		BlockOnRange = TRange<FMediaTimeStamp>(FMediaTimeStamp(Start, OnBlockPrimaryIndex, OnBlockSecondaryIndexOffset + LoopIdxS), FMediaTimeStamp(End, OnBlockPrimaryIndex, OnBlockSecondaryIndexOffset + LoopIdxE));
+		t0.AdjustLoopIndex(OnBlockLoopIndexOffset);
+		t1.AdjustLoopIndex(OnBlockLoopIndexOffset);
+		BlockOnRange = TRange<FMediaTimeStamp>(t0, t1);
 		check(!BlockOnRange.IsEmpty());
 	}
 
@@ -1505,12 +1501,9 @@ const TRange<FMediaTimeStamp>& FMediaPlayerFacade::FBlockOnRange::GetRange() con
 	// Note: Due to varying DTs the new range will NOT be a simple monotone progression in playback direction, but might overlap or even be a subset of the previous one
 	//		 We do not put any safeguards in place here, but rather use the "is last sample still valid" logic to reject illogical / impossible range requests.
 	//		 All that aside: we DO expect ranges start (lower bound if forward, upper if reverse playback) to be moving in a monotone manner according to the set playback direction.
-
 	CurrentPlayer->GetControls().SetBlockingPlaybackHint(!BlockOnRange.IsEmpty());
-
-	LastTimeRange = CurrentTimeRange;
+	LastProcessedTimeRange = CurrentTimeRange;
 	RangeIsDirty = false;
-
 	return BlockOnRange;
 }
 
@@ -2572,9 +2565,8 @@ void FMediaPlayerFacade::UpdateSeekStatus(const FMediaTimeStamp* pCheckTimeStamp
 				Player->GetSamples().PeekVideoSampleTime(VideoTimeStamp);
 			}
 
-			if (VideoTimeStamp.IsValid())
+			if (VideoTimeStamp.IsValid() && VideoTimeStamp.GetSequenceIndex() >= NextSequenceIndex.Get(0))
 			{
-				check(VideoTimeStamp.GetSequenceIndex() >= NextSequenceIndex.Get(0));
 				bool bRunningNonAudioClock = bHaveActiveAudio && !BlockOnRange.IsSet();
 				if (bRunningNonAudioClock)
 				{
diff --git a/Engine/Source/Runtime/MediaUtils/Public/MediaPlayerFacade.h b/Engine/Source/Runtime/MediaUtils/Public/MediaPlayerFacade.h
--- a/Engine/Source/Runtime/MediaUtils/Public/MediaPlayerFacade.h
+++ b/Engine/Source/Runtime/MediaUtils/Public/MediaPlayerFacade.h
@@ -997,33 +997,33 @@ private:
 		{
 			BlockOnRange = TRange<FMediaTimeStamp>::Empty();
 			CurrentTimeRange = TRange<FTimespan>::Empty();
-			LastTimeRange = TRange<FTimespan>::Empty();
+			LastProcessedTimeRange = TRange<FTimespan>::Empty();
 			RangeIsDirty = false;
-			OnBlockPrimaryIndex = 0;
-			OnBlockSecondaryIndexOffset = 0;
+			OnBlockSequenceIndex = 0;
+			OnBlockLoopIndexOffset = 0;
 		}
 
 	private:
 		/** The hosting player facade */
 		FMediaPlayerFacade* Facade;
 
-		/** The last user set time range to block on */
+		/** The most recently set time range to block on */
 		TRange<FTimespan> CurrentTimeRange;
 
+		/** Last user provided BlockOnRange value */
+		mutable TRange<FTimespan> LastProcessedTimeRange;
+
 		/** The time range to block on sample fetching. */
 		mutable TRange<FMediaTimeStamp> BlockOnRange;
 
-		/** Last user provided BlockOnRange value */
-		mutable TRange<FTimespan> LastTimeRange;
-
 		/** Flag to indicate if internal range is valid or not */
 		mutable bool RangeIsDirty;
 
 		/** Primary ("seek") sequence index used during blocked playback processing */
-		mutable int32 OnBlockPrimaryIndex;
+		mutable int32 OnBlockSequenceIndex;
 
 		/** Secondary (loop) sequence index offset used during blocked playback processing */
-		mutable int32 OnBlockSecondaryIndexOffset;
+		mutable int32 OnBlockLoopIndexOffset;
 	};
 
 	FBlockOnRange BlockOnRange;
