Media: Console variable to control block on the block on range timeout duration, previously hard-coded to 10 seconds.

This can help facilitate debugging.

diff --git a/Engine/Source/Runtime/MediaUtils/Private/MediaPlayerFacade.cpp b/Engine/Source/Runtime/MediaUtils/Private/MediaPlayerFacade.cpp
--- a/Engine/Source/Runtime/MediaUtils/Private/MediaPlayerFacade.cpp
+++ b/Engine/Source/Runtime/MediaUtils/Private/MediaPlayerFacade.cpp
@@ -108,6 +108,13 @@ TAutoConsoleVariable<bool> CVarTestForcePlayerCreateFailed(
 
 #endif // !UE_BUILD_SHIPPING
 
+float GBlockOnFetchTimeout = 10.0f;
+static FAutoConsoleVariableRef CVarMediaUtilsBlockOnFetchTimeout(
+	TEXT("MediaUtils.BlockOnFetchTimeout"),
+	GBlockOnFetchTimeout,
+	TEXT("Maximum time that TickInput/Fetch will block waiting for samples (in seconds).\n")
+);
+
 }
 
 /* Local helpers
@@ -2305,7 +2312,7 @@ void FMediaPlayerFacade::TickInput(FTimespan DeltaTime, FTimespan Timecode)
 				}
 
 				// Timeout?
-				if ((FPlatformTime::Seconds() - BlockingStart) > MEDIAUTILS_MAX_BLOCKONFETCH_SECONDS)
+				if ((FPlatformTime::Seconds() - BlockingStart) > static_cast<double>(UE::MediaUtils::Private::GBlockOnFetchTimeout))
 				{
 					FString Url;
 #if !UE_BUILD_SHIPPING
@@ -2410,6 +2417,7 @@ void FMediaPlayerFacade::TickFetch(FTimespan DeltaTime, FTimespan Timecode)
 		// process samples in range
 		IMediaSamples& Samples = CurrentPlayer->GetSamples();
 
+		const double BlockOnFetchTimeout = static_cast<double>(UE::MediaUtils::Private::GBlockOnFetchTimeout);
 		bool Blocked = false;
 		FDateTime BlockedTime;
 
@@ -2426,11 +2434,11 @@ void FMediaPlayerFacade::TickFetch(FTimespan DeltaTime, FTimespan Timecode)
 
 			if (Blocked)
 			{
-				if ((FDateTime::UtcNow() - BlockedTime) >= FTimespan::FromSeconds(MEDIAUTILS_MAX_BLOCKONFETCH_SECONDS))
+				if ((FDateTime::UtcNow() - BlockedTime) >= FTimespan::FromSeconds(BlockOnFetchTimeout))
 				{
 					UE_LOG(LogMediaUtils, Verbose, TEXT("PlayerFacade: Aborted block on fetch %s after %i seconds"),
 						*BlockOnRange.GetRange().GetLowerBoundValue().Time.ToString(TEXT("%h:%m:%s.%t")),
-						MEDIAUTILS_MAX_BLOCKONFETCH_SECONDS
+						static_cast<int32>(BlockOnFetchTimeout)
 					);
 
 					break;
diff --git a/Engine/Source/Runtime/MediaUtils/Private/MediaUtilsPrivate.h b/Engine/Source/Runtime/MediaUtils/Private/MediaUtilsPrivate.h
--- a/Engine/Source/Runtime/MediaUtils/Private/MediaUtilsPrivate.h
+++ b/Engine/Source/Runtime/MediaUtils/Private/MediaUtilsPrivate.h
@@ -5,10 +5,5 @@
 #include "CoreMinimal.h"
 #include "Stats/Stats.h"
 
-
-/** Maximum time that TickInput/Fetch will block waiting for samples (in seconds). */
-#define MEDIAUTILS_MAX_BLOCKONFETCH_SECONDS 10
-
-
 /** Declares a log category for this module. */
 DECLARE_LOG_CATEGORY_EXTERN(LogMediaUtils, Log, All);
