MediaUtils: Fixed use of an unset TOptional in non-v2 timing players

#rnx

diff --git a/Engine/Source/Runtime/MediaUtils/Private/MediaPlayerFacade.cpp b/Engine/Source/Runtime/MediaUtils/Private/MediaPlayerFacade.cpp
--- a/Engine/Source/Runtime/MediaUtils/Private/MediaPlayerFacade.cpp
+++ b/Engine/Source/Runtime/MediaUtils/Private/MediaPlayerFacade.cpp
@@ -1770,7 +1770,7 @@ void FMediaPlayerFacade::Flush(bool bExcludePlayer, bool bOnSeek)
 	LastAudioRenderedSampleTime.Invalidate();
 	if (bOnSeek)
 	{
-		BlockOnRange.OnSeek(NextSequenceIndex.GetValue());
+		BlockOnRange.OnSeek(NextSequenceIndex.Get(0));
 	}
 	else
 	{
