[Media] Adding some diagnostic verbose logging to media player and player facade to help troubleshoot issues with BlockOnRange and PlaybackRange.

diff --git a/Engine/Source/Runtime/MediaAssets/Private/Assets/MediaPlayer.cpp b/Engine/Source/Runtime/MediaAssets/Private/Assets/MediaPlayer.cpp
--- a/Engine/Source/Runtime/MediaAssets/Private/Assets/MediaPlayer.cpp
+++ b/Engine/Source/Runtime/MediaAssets/Private/Assets/MediaPlayer.cpp
@@ -825,7 +825,8 @@ void UMediaPlayer::SetBlockOnTime(const FTimespan& Time)
 
 void UMediaPlayer::SetBlockOnTimeRange(const TRange<FTimespan>& TimeRange)
 {
-	UE_LOG(LogMediaAssets, VeryVerbose, TEXT("%s.SetBlockOnRange %s"), *GetFName().ToString(), *TimeRange.GetLowerBoundValue().ToString(TEXT("%h:%m:%s.%t")));
+	UE_LOG(LogMediaAssets, VeryVerbose, TEXT("%s.SetBlockOnRange [%s, %s]"), *GetFName().ToString(),
+		*TimeRange.GetLowerBoundValue().ToString(TEXT("%h:%m:%s.%t")), *TimeRange.GetUpperBoundValue().ToString(TEXT("%h:%m:%s.%t")));
 	return PlayerFacade->SetBlockOnTimeRange(TimeRange);
 }
 
@@ -936,6 +937,8 @@ TRange<FTimespan> UMediaPlayer::GetPlaybackTimeRange(EMediaTimeRangeType InRange
 
 bool UMediaPlayer::SetPlaybackTimeRange(const TRange<FTimespan>& InTimeRange)
 {
+	UE_LOG(LogMediaAssets, VeryVerbose, TEXT("%s.SetPlaybackTimeRange [%s, %s]"), *GetFName().ToString(),
+		*InTimeRange.GetLowerBoundValue().ToString(TEXT("%h:%m:%s.%t")), *InTimeRange.GetUpperBoundValue().ToString(TEXT("%h:%m:%s.%t")));
 	return PlayerFacade->SetPlaybackTimeRange(InTimeRange);
 }
 
diff --git a/Engine/Source/Runtime/MediaUtils/Private/MediaPlayerFacade.cpp b/Engine/Source/Runtime/MediaUtils/Private/MediaPlayerFacade.cpp
--- a/Engine/Source/Runtime/MediaUtils/Private/MediaPlayerFacade.cpp
+++ b/Engine/Source/Runtime/MediaUtils/Private/MediaPlayerFacade.cpp
@@ -2318,7 +2318,9 @@ void FMediaPlayerFacade::TickInput(FTimespan DeltaTime, FTimespan Timecode)
 #if !UE_BUILD_SHIPPING
 					Url = Player->GetUrl();
 #endif // !UE_BUILD_SHIPPING
-					UE_LOG(LogMediaUtils, Error, TEXT("Blocking media playback timed out. Disabling it for this playback session. URL:%s"),
+					UE_LOG(LogMediaUtils, Error, TEXT("Blocking media playback timed out on range [%s (seq:%d, loop:%d), %s (seq:%d, loop:%d)]. Disabling it for this playback session. URL:%s"),
+						*TimeRange.GetLowerBoundValue().Time.ToString(TEXT("%h:%m:%s.%t")), TimeRange.GetLowerBoundValue().GetSequenceIndex(), TimeRange.GetLowerBoundValue().GetLoopIndex(),
+						*TimeRange.GetUpperBoundValue().Time.ToString(TEXT("%h:%m:%s.%t")), TimeRange.GetUpperBoundValue().GetSequenceIndex(), TimeRange.GetUpperBoundValue().GetLoopIndex(),
 						*Url);
 					BlockOnRangeDisabled = true;
 					break;
@@ -2609,9 +2611,11 @@ void FMediaPlayerFacade::ExecuteNextSeek()
 {
 	if (NextSeekTime.IsSet() && !IsSeeking())
 	{
+		UE_LOG(LogMediaUtils, Verbose, TEXT("PlayerFacade: Internal Seek at %s."), *NextSeekTime.GetValue().ToString(TEXT("%h:%m:%s.%t")));
+
 		if (!Seek(NextSeekTime.GetValue()))
 		{
-			// todo: signal error for failed seek.
+			UE_LOG(LogMediaUtils, Error, TEXT("PlayerFacade: Failed Internal Seek at %s."), *NextSeekTime.GetValue().ToString(TEXT("%h:%m:%s.%t")));
 		}
 	}
 }
@@ -3277,6 +3281,19 @@ bool FMediaPlayerFacade::ProcessVideoSamples(IMediaSamples& Samples, const TRang
 		FMediaTimeStamp SampleTime = Sample->GetTime();
 		TRange<FMediaTimeStamp> SampleTimeRange(SampleTime, SampleTime + Sample->GetDuration());
 
+#if !NO_LOGGING
+		if (VideoSampleSinks.IsEmpty())
+		{
+			UE_LOG(LogMediaUtils, Verbose, TEXT("PlayerFacade: No Video Sample Sinks available, Dropping Sample %s Sequence %d Loop %d"),
+				*SampleTime.GetTime().ToString(TEXT("%h:%m:%s.%t")), SampleTime.GetSequenceIndex(), SampleTime.GetLoopIndex());
+		}
+		else
+		{
+			UE_LOG(LogMediaUtils, VeryVerbose, TEXT("PlayerFacade: Enqueuing Sample %s Sequence %d Loop %d"),
+				*SampleTime.GetTime().ToString(TEXT("%h:%m:%s.%t")), SampleTime.GetSequenceIndex(), SampleTime.GetLoopIndex());
+		}
+#endif
+
 		// Enqueue the sample to render
 		// (we use a queue to stay compatible with existing structure and older sinks - new sinks will read this single entry right away on the gamethread
 		//  and pass it along to rendering outside the queue)
