Checking block-on-range for valid range (upper bound not less than lower bound)

#rnx

diff --git a/Engine/Source/Runtime/MediaUtils/Private/MediaPlayerFacade.cpp b/Engine/Source/Runtime/MediaUtils/Private/MediaPlayerFacade.cpp
--- a/Engine/Source/Runtime/MediaUtils/Private/MediaPlayerFacade.cpp
+++ b/Engine/Source/Runtime/MediaUtils/Private/MediaPlayerFacade.cpp
@@ -1367,8 +1367,15 @@ void FMediaPlayerFacade::FBlockOnRange::SetRange(const TRange<FTimespan>& NewRan
 {
 	if (CurrentTimeRange != NewRange)
 	{
-		CurrentTimeRange = NewRange;
-		RangeIsDirty = true;
+		if (CurrentTimeRange.GetLowerBoundValue().GetTicks() <= CurrentTimeRange.GetUpperBoundValue().GetTicks())
+		{
+			CurrentTimeRange = NewRange;
+			RangeIsDirty = true;
+		}
+		else
+		{
+			Reset();
+		}
 	}
 }
 
@@ -1376,28 +1383,28 @@ void FMediaPlayerFacade::FBlockOnRange::SetRange(const TRange<FTimespan>& NewRan
 bool FMediaPlayerFacade::FBlockOnRange::IsSet() const
 {
 	TSharedPtr<IMediaPlayer, ESPMode::ThreadSafe> CurrentPlayer(Facade->Player);
-	check(CurrentPlayer.IsValid());
 
 	if (!RangeIsDirty)
 	{
 		return !BlockOnRange.IsEmpty();
 	}
-	return (!CurrentTimeRange.IsEmpty() && CurrentPlayer->GetControls().CanControl(EMediaControl::BlockOnFetch));
+	return !CurrentTimeRange.IsEmpty() && CurrentPlayer.IsValid() && CurrentPlayer->GetControls().CanControl(EMediaControl::BlockOnFetch);
 }
 
 
 const TRange<FMediaTimeStamp>& FMediaPlayerFacade::FBlockOnRange::GetRange() const
 {
 	TSharedPtr<IMediaPlayer, ESPMode::ThreadSafe> CurrentPlayer(Facade->Player);
-	check(CurrentPlayer.IsValid());
 
-	if (!RangeIsDirty)
+	if (!RangeIsDirty || !CurrentPlayer.IsValid())
 	{
 		return BlockOnRange;
 	}
 
-	// If the range is empty or the player can't support blocked playback: reset everything & return empty block range...
-	if (CurrentTimeRange.IsEmpty() || !CurrentPlayer->GetControls().CanControl(EMediaControl::BlockOnFetch))
+	// If the range is empty or the player can't support blocked playback or the range interval is invalid: reset everything & return empty block range...
+	if (CurrentTimeRange.IsEmpty() || !CurrentPlayer->GetControls().CanControl(EMediaControl::BlockOnFetch) ||
+		CurrentTimeRange.GetLowerBoundValue().GetTicks() < CurrentTimeRange.GetUpperBoundValue().GetTicks() ||
+		CurrentPlayer->GetControls().GetDuration().GetTicks() <= 0)
 	{
 		LastProcessedTimeRange = TRange<FTimespan>::Empty();
 		BlockOnRange = TRange<FMediaTimeStamp>::Empty();
@@ -1417,7 +1424,6 @@ const TRange<FMediaTimeStamp>& FMediaPlayerFacade::FBlockOnRange::GetRange() con
 	const int64 StartTicks = CurrentTimeRange.GetLowerBoundValue().GetTicks();
 	const int64 EndExclusiveTicks = CurrentTimeRange.GetUpperBoundValue().GetTicks();
 	const int64 DurationTicks = CurrentPlayer->GetControls().GetDuration().GetTicks();
-	check(StartTicks >= 0 && EndExclusiveTicks >= StartTicks);
 
 	/*
 		When looping we need to synthesize the expected start and end loop indices of the range we are returning.
@@ -1465,7 +1471,6 @@ const TRange<FMediaTimeStamp>& FMediaPlayerFacade::FBlockOnRange::GetRange() con
 				a movie repetition (ie when the movie has been pulled out in the sequencer track to repeat a number of times).
 				We set that repetion count as the base for the loop index.
 			*/
-			check(OnBlockLoopIndexOffset == 0);
 			OnBlockLoopIndexOffset = !bReverse ? -t0.GetLoopIndex() : -t1.GetLoopIndex();
 		}
 		else
