[MediaUtils] PlayerFacade: Fix inverted block on range condition.

Regression introduced in CL 44623878.

diff --git a/Engine/Source/Runtime/MediaUtils/Private/MediaPlayerFacade.cpp b/Engine/Source/Runtime/MediaUtils/Private/MediaPlayerFacade.cpp
--- a/Engine/Source/Runtime/MediaUtils/Private/MediaPlayerFacade.cpp
+++ b/Engine/Source/Runtime/MediaUtils/Private/MediaPlayerFacade.cpp
@@ -1403,7 +1403,7 @@ const TRange<FMediaTimeStamp>& FMediaPlayerFacade::FBlockOnRange::GetRange() con
 
 	// If the range is empty or the player can't support blocked playback or the range interval is invalid: reset everything & return empty block range...
 	if (CurrentTimeRange.IsEmpty() || !CurrentPlayer->GetControls().CanControl(EMediaControl::BlockOnFetch) ||
-		CurrentTimeRange.GetLowerBoundValue().GetTicks() < CurrentTimeRange.GetUpperBoundValue().GetTicks() ||
+		CurrentTimeRange.GetLowerBoundValue().GetTicks() > CurrentTimeRange.GetUpperBoundValue().GetTicks() ||
 		CurrentPlayer->GetControls().GetDuration().GetTicks() <= 0)
 	{
 		LastProcessedTimeRange = TRange<FTimespan>::Empty();
