[MediaUtils] MediaPlayerFacade: Fix incorrect condition check in SetRande.

Potential regression introduced in CL 44623878.
Spotted by BugHawk.

diff --git a/Engine/Source/Runtime/MediaUtils/Private/MediaPlayerFacade.cpp b/Engine/Source/Runtime/MediaUtils/Private/MediaPlayerFacade.cpp
--- a/Engine/Source/Runtime/MediaUtils/Private/MediaPlayerFacade.cpp
+++ b/Engine/Source/Runtime/MediaUtils/Private/MediaPlayerFacade.cpp
@@ -1367,7 +1367,7 @@ void FMediaPlayerFacade::FBlockOnRange::SetRange(const TRange<FTimespan>& NewRan
 {
 	if (CurrentTimeRange != NewRange)
 	{
-		if (CurrentTimeRange.GetLowerBoundValue().GetTicks() <= CurrentTimeRange.GetUpperBoundValue().GetTicks())
+		if (NewRange.GetLowerBoundValue().GetTicks() <= NewRange.GetUpperBoundValue().GetTicks())
 		{
 			CurrentTimeRange = NewRange;
 			RangeIsDirty = true;
