[MediaUtils] MediaPlayerFacade BlockOnRange: More Range validation error handling

Range validity also includes if the range is closed. Not handling this condition may lead to asserts. Reported by BugHawk.

diff --git a/Engine/Source/Runtime/MediaUtils/Private/MediaPlayerFacade.cpp b/Engine/Source/Runtime/MediaUtils/Private/MediaPlayerFacade.cpp
--- a/Engine/Source/Runtime/MediaUtils/Private/MediaPlayerFacade.cpp
+++ b/Engine/Source/Runtime/MediaUtils/Private/MediaPlayerFacade.cpp
@@ -1367,7 +1367,7 @@ void FMediaPlayerFacade::FBlockOnRange::SetRange(const TRange<FTimespan>& NewRan
 {
 	if (CurrentTimeRange != NewRange)
 	{
-		if (NewRange.GetLowerBoundValue().GetTicks() <= NewRange.GetUpperBoundValue().GetTicks())
+		if (NewRange.HasLowerBound() && NewRange.HasUpperBound() && NewRange.GetLowerBoundValue().GetTicks() <= NewRange.GetUpperBoundValue().GetTicks())
 		{
 			CurrentTimeRange = NewRange;
 			RangeIsDirty = true;
