MediaUtils: Unifying which critical section to use when accessing time values since there were discrepancies causing occasional race condition while scrubbing

#rnx

diff --git a/Engine/Source/Runtime/MediaUtils/Private/MediaPlayerFacade.cpp b/Engine/Source/Runtime/MediaUtils/Private/MediaPlayerFacade.cpp
--- a/Engine/Source/Runtime/MediaUtils/Private/MediaPlayerFacade.cpp
+++ b/Engine/Source/Runtime/MediaUtils/Private/MediaPlayerFacade.cpp
@@ -523,13 +523,13 @@ FMediaTimeStamp FMediaPlayerFacade::GetDisplayTimeStamp() const
 
 TOptional<FTimecode> FMediaPlayerFacade::GetVideoTimecode() const
 {
-	FScopeLock Lock(&LastTimeValuesCS);
+	FScopeLock Lock(&TimeValueLock);
 	return MostRecentlyDeliveredVideoFrameTimecode;
 }
 
 TRange<FMediaTimeStamp> FMediaPlayerFacade::GetLastProcessedVideoSampleTimeRange() const
 {
-	FScopeLock Lock(&LastTimeValuesCS);
+	FScopeLock Lock(&TimeValueLock);
 	return LastVideoSampleProcessedTimeRange;
 }
 
@@ -541,16 +541,15 @@ FMediaTimeStamp FMediaPlayerFacade::GetTimeStampInternal(bool bForDisplay) const
 		return FMediaTimeStamp();
 	}
 
-	FScopeLock Lock(&LastTimeValuesCS);
-
 	if (!CurrentPlayer->GetPlayerFeatureFlag(IMediaPlayer::EFeatureFlag::UsePlaybackTimingV2))
 	{
 		// Make sure we can return values for V1 players...
 		return FMediaTimeStamp(GetTime());
 	}
 
+	FScopeLock Lock(&TimeValueLock);
 	// Check if the value is for display purposes. If so: do we seek right now?
-	if (bForDisplay && SeekTargetTime.IsValid())
+	if (bForDisplay && IsSeeking())
 	{
 		return SeekTargetTime;
 	}
@@ -1164,7 +1163,7 @@ bool FMediaPlayerFacade::ContinueOpen(IMediaPlayerLifecycleManagerDelegate::ICon
 	}
 
 	{
-		FScopeLock Lock(&LastTimeValuesCS);
+		FScopeLock Lock(&TimeValueLock);
 
 		BlockOnRangeDisabled = false;
 		BlockOnRange.OnFlush();
@@ -1225,7 +1224,9 @@ void FMediaPlayerFacade::QueryCacheState(EMediaTrackType TrackType, EMediaCacheS
 
 bool FMediaPlayerFacade::Seek(const FTimespan& InTime)
 {
+	TimeValueLock.Lock();
 	NextSeekTime.Reset();
+	TimeValueLock.Unlock();
 
 	auto CurrentPlayer = Player;
 
@@ -1271,8 +1272,8 @@ bool FMediaPlayerFacade::Seek(const FTimespan& InTime)
 		return false;
 	}
 
+	{
 	FScopeLock Lock(&CriticalSection);
-
 	// V2 timing player?
 	if (CurrentPlayer->GetPlayerFeatureFlag(IMediaPlayer::EFeatureFlag::UsePlaybackTimingV2))
 	{
@@ -1289,28 +1290,29 @@ bool FMediaPlayerFacade::Seek(const FTimespan& InTime)
 			Flush(CurrentPlayer->GetPlayerFeatureFlag(IMediaPlayer::EFeatureFlag::PlayerUsesInternalFlushOnSeek), false);
 		}
 	}
+	}
 
+	TimeValueLock.Lock();
 	SeekTargetTime = FMediaTimeStamp(Time, NextSequenceIndex.Get(0), 0);
+	TimeValueLock.Unlock();
 	return true;
 }
 
 bool FMediaPlayerFacade::IsSeeking() const
 {
-	// Code inspection notes:
-	// - Usually protected by LastTimeValuesCS, but sometimes CriticalSection (or both interlocked).
-	// - GetCurrentPlaybackTimeRange reads it outside of a scope lock.
-	FScopeLock Lock(&LastTimeValuesCS);
+	FScopeLock Lock(&TimeValueLock);
 	return SeekTargetTime.IsValid();
 }
 
 FMediaTimeStamp FMediaPlayerFacade::GetSeekTarget() const
 {
-	FScopeLock Lock(&LastTimeValuesCS);
+	FScopeLock Lock(&TimeValueLock);
 	return SeekTargetTime;
 }
 
 void FMediaPlayerFacade::SetNextSeek(const FTimespan& InTime)
 {
+	FScopeLock Lock(&TimeValueLock);
 	NextSeekTime = InTime;
 }
 
@@ -1661,14 +1663,14 @@ bool FMediaPlayerFacade::SetPlaybackTimeRange(const TRange<FTimespan>& InTimeRan
 
 void FMediaPlayerFacade::SetLastAudioRenderedSampleTime(FTimespan SampleTime)
 {
-	FScopeLock Lock(&LastTimeValuesCS);
+	FScopeLock Lock(&TimeValueLock);
 	LastAudioRenderedSampleTime.TimeStamp = FMediaTimeStamp(SampleTime);
 	LastAudioRenderedSampleTime.SampledAtTime = FPlatformTime::Seconds();
 }
 
 FTimespan FMediaPlayerFacade::GetLastAudioRenderedSampleTime() const
 {
-	FScopeLock Lock(&LastTimeValuesCS);
+	FScopeLock Lock(&TimeValueLock);
 	return LastAudioRenderedSampleTime.TimeStamp.Time;
 }
 
@@ -1805,7 +1807,9 @@ void FMediaPlayerFacade::Flush(bool bExcludePlayer, bool bOnSeek)
 		// Invalidate next video time to fetch (non-audio case)
 		NextEstVideoTimeAtFrameStart.Invalidate();
 		// ...and seek target
+		TimeValueLock.Lock();
 		SeekTargetTime.Invalidate();
+		TimeValueLock.Unlock();
 	}
 
 	NextVideoSampleTime = FTimespan::MinValue();
@@ -2344,7 +2348,7 @@ void FMediaPlayerFacade::TickInput(FTimespan DeltaTime, FTimespan Timecode)
 			if (bHaveActiveAudio)
 			{
 				// Keep currently last processed audio sample timestamp available for all frame (to provide consistent info)
-				FScopeLock Lock(&LastTimeValuesCS);
+				FScopeLock Lock(&TimeValueLock);
 				CurrentFrameAudioTimeStamp = LastAudioSampleProcessedTime.TimeStamp;
 			}
 		}
@@ -2501,7 +2505,7 @@ void FMediaPlayerFacade::TickTickable()
 	}
 
 	{
-		FScopeLock Lock1(&LastTimeValuesCS);
+		FScopeLock Lock1(&TimeValueLock);
 		Player->SetLastAudioRenderedSampleTime(LastAudioRenderedSampleTime.TimeStamp.Time);
 	}
 
@@ -2567,7 +2571,7 @@ void FMediaPlayerFacade::UpdateSeekStatus(const FMediaTimeStamp* pCheckTimeStamp
 
 	if (HaveVideoPlayback())
 	{
-		if (SeekTargetTime.IsValid())
+		if (IsSeeking())
 		{
 			// Either peek for the newest available sample or take a given timestamp to check against
 			FMediaTimeStamp VideoTimeStamp;
@@ -2587,7 +2591,7 @@ void FMediaPlayerFacade::UpdateSeekStatus(const FMediaTimeStamp* pCheckTimeStamp
 				{
 					NextEstVideoTimeAtFrameStart = FMediaTimeStampSample(VideoTimeStamp, FPlatformTime::Seconds());
 				}
-				FScopeLock LockLT(&LastTimeValuesCS);
+				FScopeLock LockLT(&TimeValueLock);
 				CurrentFrameVideoDisplayTimeStamp = SeekTargetTime;
 				SeekTargetTime.Invalidate();
 			}
@@ -2595,9 +2599,9 @@ void FMediaPlayerFacade::UpdateSeekStatus(const FMediaTimeStamp* pCheckTimeStamp
 	}
 	else
 	{
+		FScopeLock LockLT(&TimeValueLock);
 		if (bHaveActiveAudio)
 		{
-			FScopeLock LockLT(&LastTimeValuesCS);
 			if (CurrentFrameAudioTimeStamp >= SeekTargetTime)
 			{
 				SeekTargetTime.Invalidate();
@@ -2614,13 +2618,17 @@ void FMediaPlayerFacade::UpdateSeekStatus(const FMediaTimeStamp* pCheckTimeStamp
 
 void FMediaPlayerFacade::ExecuteNextSeek()
 {
-	if (NextSeekTime.IsSet() && !IsSeeking())
+	TimeValueLock.Lock();
+	auto nst = NextSeekTime;
+	TimeValueLock.Unlock();
+
+	if (nst.IsSet() && !IsSeeking())
 	{
-		UE_LOG(LogMediaUtils, Verbose, TEXT("PlayerFacade: Internal Seek at %s."), *NextSeekTime.GetValue().ToString(TEXT("%h:%m:%s.%t")));
+		UE_LOG(LogMediaUtils, Verbose, TEXT("PlayerFacade: Internal Seek at %s."), *nst.GetValue().ToString(TEXT("%h:%m:%s.%t")));
 
-		if (!Seek(NextSeekTime.GetValue()))
+		if (!Seek(nst.GetValue()))
 		{
-			UE_LOG(LogMediaUtils, Error, TEXT("PlayerFacade: Failed Internal Seek at %s."), *NextSeekTime.GetValue().ToString(TEXT("%h:%m:%s.%t")));
+			UE_LOG(LogMediaUtils, Error, TEXT("PlayerFacade: Failed Internal Seek at %s."), *nst.GetValue().ToString(TEXT("%h:%m:%s.%t")));
 		}
 	}
 }
@@ -2650,7 +2658,7 @@ void FMediaPlayerFacade::PreSampleProcessingTimeHandling()
 	UpdateSeekStatus();
 
 	// No seeking?
-	if (!SeekTargetTime.IsValid())
+	if (!IsSeeking())
 	{
 		// No seek pending & not paused. Can we / Do we need to prime a non-audio clock?
 		if (!bHaveActiveAudio && !BlockOnRange.IsSet())
@@ -2837,7 +2845,7 @@ bool FMediaPlayerFacade::GetCurrentPlaybackTimeRange(TRange<FMediaTimeStamp>& Ti
 			if (!bPurgeSampleRelated)
 			{
 				// If paused and not seeking, make sure we get one sample nonetheless...
-				if (Rate == 0.0f && !SeekTargetTime.IsValid())
+				if (Rate == 0.0f && !IsSeeking())
 				{
 					// Do this once after open / seek...
 					if (LastVideoSampleProcessedTimeRange.IsEmpty())
@@ -3026,7 +3034,7 @@ void FMediaPlayerFacade::ProcessAudioSamples(IMediaSamples& Samples, const TRang
 			}
 
 			{
-				FScopeLock Lock(&LastTimeValuesCS);
+				FScopeLock Lock(&TimeValueLock);
 				LastAudioSampleProcessedTime.TimeStamp = FMediaTimeStamp(Sample->GetTime());
 				LastAudioSampleProcessedTime.SampledAtTime = FPlatformTime::Seconds();
 			}
@@ -3044,7 +3052,7 @@ void FMediaPlayerFacade::ProcessAudioSamples(IMediaSamples& Samples, const TRang
 			// (rough estimate, as this is off-gamethread; but better than throwing things out with no throttling at all)
 			{
 				bool bReverse = (CurrentRate < 0.0f);
-				FScopeLock Lock(&LastTimeValuesCS);
+				FScopeLock Lock(&TimeValueLock);
 				if (!bReverse)
 				{
 					TempRange.SetUpperBound(CurrentFrameVideoTimeStamp);
@@ -3073,7 +3081,7 @@ void FMediaPlayerFacade::ProcessAudioSamples(IMediaSamples& Samples, const TRang
 bool FMediaPlayerFacade::IsVideoSampleStillGood(const TRange<FMediaTimeStamp>& LastSampleTimeRange, const TRange<FMediaTimeStamp>& TimeRange, bool bReverse) const
 {
 	// If we have no valid time range or a seek is in progress we assume the current frame can be considered "done" in any case
-	if (!TimeRange.IsEmpty() && !SeekTargetTime.IsValid() && !LastSampleTimeRange.IsEmpty())
+	if (!TimeRange.IsEmpty() && !IsSeeking() && !LastSampleTimeRange.IsEmpty())
 	{
 		// This is not the case: check more detailed!
 
@@ -3306,7 +3314,7 @@ bool FMediaPlayerFacade::ProcessVideoSamples(IMediaSamples& Samples, const TRang
 		check(bOk);
 
 		{
-			FScopeLock Lock(&LastTimeValuesCS);
+			FScopeLock Lock(&TimeValueLock);
 			CurrentFrameVideoDisplayTimeStamp = CurrentFrameVideoTimeStamp = SampleTimeRange.GetLowerBoundValue();
 			LastVideoSampleProcessedTimeRange = SampleTimeRange;
 			MostRecentlyDeliveredVideoFrameTimecode = Sample->GetTimecode();
@@ -3326,10 +3334,11 @@ void FMediaPlayerFacade::ProcessCaptionSamples(IMediaSamples& Samples, const TRa
 	TSharedPtr<IMediaOverlaySample, ESPMode::ThreadSafe> Sample;
 
 	// Seek in progress?
-	if (SeekTargetTime.IsValid())
+	if (IsSeeking())
 	{
 		// Yes. Fetch (and discard) all samples up to the seek target time...
 		// (we only throw out samples from prior sequence indices to make sure we do not swallow any audio from overlapping samples)
+		FScopeLock Lock(&TimeValueLock);
 		auto DiscardRange = TRange<FMediaTimeStamp>(FMediaTimeStamp(0), FMediaTimeStamp((CurrentRate >= 0.0f) ? FTimespan::Zero() : FTimespan::MaxValue(), SeekTargetTime.GetSequenceIndex(), SeekTargetTime.GetLoopIndex()));
 		Samples.DiscardCaptionSamples(DiscardRange, GetUnpausedRate() < 0.0f);
 	}
@@ -3355,10 +3364,11 @@ void FMediaPlayerFacade::ProcessSubtitleSamples(IMediaSamples& Samples, const TR
 	TSharedPtr<IMediaOverlaySample, ESPMode::ThreadSafe> Sample;
 
 	// Seek in progress?
-	if (SeekTargetTime.IsValid())
+	if (IsSeeking())
 	{
 		// Yes. Fetch (and discard) all samples up to the seek target time...
 		// (we only throw out samples from prior sequence indices to make sure we do not swallow any audio from overlapping samples)
+		FScopeLock Lock(&TimeValueLock);
 		auto DiscardRange = TRange<FMediaTimeStamp>(FMediaTimeStamp(0), FMediaTimeStamp((CurrentRate >= 0.0f) ? FTimespan::Zero() : FTimespan::MaxValue(), SeekTargetTime.GetSequenceIndex(), SeekTargetTime.GetLoopIndex()));
 		Samples.DiscardSubtitleSamples(DiscardRange, GetUnpausedRate() < 0.0f);
 	}
@@ -3385,10 +3395,11 @@ void FMediaPlayerFacade::ProcessMetadataSamples(IMediaSamples& Samples, const TR
 	TSharedPtr<IMediaBinarySample, ESPMode::ThreadSafe> Sample;
 
 	// Seek in progress?
-	if (SeekTargetTime.IsValid())
+	if (IsSeeking())
 	{
 		// Yes. Fetch (and discard) all samples up to the seek target time...
 		// (we only throw out samples from prior sequence indices to make sure we do not swallow any audio from overlapping samples)
+		FScopeLock Lock(&TimeValueLock);
 		auto DiscardRange = TRange<FMediaTimeStamp>(FMediaTimeStamp(0), FMediaTimeStamp((CurrentRate >= 0.0f) ? FTimespan::Zero() : FTimespan::MaxValue(), SeekTargetTime.GetSequenceIndex(), SeekTargetTime.GetLoopIndex()));
 		Samples.DiscardMetadataSamples(DiscardRange, GetUnpausedRate() < 0.0f);
 	}
@@ -3426,7 +3437,7 @@ void FMediaPlayerFacade::ProcessAudioSamplesV1(IMediaSamples& Samples, TRange<FT
 		}
 		else
 		{
-			FScopeLock Lock(&LastTimeValuesCS);
+			FScopeLock Lock(&TimeValueLock);
 			LastAudioSampleProcessedTime.TimeStamp = Sample->GetTime();
 			LastAudioSampleProcessedTime.SampledAtTime = FPlatformTime::Seconds();
 		}
@@ -3449,7 +3460,7 @@ void FMediaPlayerFacade::ProcessVideoSamplesV1(IMediaSamples& Samples, TRange<FT
 		}
 
 		{
-			FScopeLock Lock(&LastTimeValuesCS);
+			FScopeLock Lock(&TimeValueLock);
 			CurrentFrameVideoDisplayTimeStamp = CurrentFrameVideoTimeStamp = Sample->GetTime();
 		}
 
@@ -3673,7 +3684,7 @@ void FMediaPlayerFacade::ReInit()
 	}
 
 	{
-		FScopeLock lock(&LastTimeValuesCS);
+		FScopeLock lock(&TimeValueLock);
 		LastAudioRenderedSampleTime.Invalidate();
 		LastAudioSampleProcessedTime.Invalidate();
 		LastVideoSampleProcessedTimeRange = TRange<FMediaTimeStamp>::Empty();
diff --git a/Engine/Source/Runtime/MediaUtils/Public/MediaPlayerFacade.h b/Engine/Source/Runtime/MediaUtils/Public/MediaPlayerFacade.h
--- a/Engine/Source/Runtime/MediaUtils/Public/MediaPlayerFacade.h
+++ b/Engine/Source/Runtime/MediaUtils/Public/MediaPlayerFacade.h
@@ -625,13 +625,13 @@ public:
 	 * @see Seek
 	 */
 	MEDIAUTILS_API bool IsSeeking() const;
-	
+
 	/**
 	 * If valid, returns the current seek request.
 	 * @see Seek, IsSeeking
 	 */
 	MEDIAUTILS_API FMediaTimeStamp GetSeekTarget() const;
-	
+
 	/**
 	 *	Set seek to be done after current seek is completed.
 	 *	@remark This value is reset if Seek is called directly.
@@ -1079,7 +1079,7 @@ private:
 	TQueue<EMediaEvent, EQueueMode::Spsc> QueuedEventBroadcasts;
 
 	/** CS to make last time values below thread safe */
-	mutable FCriticalSection LastTimeValuesCS;
+	mutable FCriticalSection TimeValueLock;
 
 	/** Time of last audio sample played. */
 	FMediaTimeStampSample LastAudioRenderedSampleTime;
