MediaUtils: Added optional tracking of dropped media samples

#rnx

diff --git a/Engine/Source/Runtime/Media/Public/IMediaSamples.h b/Engine/Source/Runtime/Media/Public/IMediaSamples.h
--- a/Engine/Source/Runtime/Media/Public/IMediaSamples.h
+++ b/Engine/Source/Runtime/Media/Public/IMediaSamples.h
@@ -161,6 +161,12 @@ public:
 	virtual int32 NumSubtitleSamples() const { return -1; }
 	virtual int32 NumVideoSamples() const { return -1; }
 
+	virtual uint32 GetNumDroppedVideoSamples(bool bInClearToZero) { return 0U; }
+	virtual uint32 GetNumDroppedAudioSamples(bool bInClearToZero) { return 0U; }
+	virtual uint32 GetNumDroppedSubtitleSamples(bool bInClearToZero) { return 0U; }
+	virtual uint32 GetNumDroppedCaptionSamples(bool bInClearToZero) { return 0U; }
+	virtual uint32 GetNumDroppedMetadataSamples(bool bInClearToZero) { return 0U; }
+
 public:
 
 	/** Virtual destructor. */
diff --git a/Engine/Source/Runtime/MediaUtils/Public/MediaSampleQueue.h b/Engine/Source/Runtime/MediaUtils/Public/MediaSampleQueue.h
--- a/Engine/Source/Runtime/MediaUtils/Public/MediaSampleQueue.h
+++ b/Engine/Source/Runtime/MediaUtils/Public/MediaSampleQueue.h
@@ -63,6 +63,25 @@ public:
 		return Samples.Num();
 	}
 
+	/**
+	 * Returns the number of samples that were dropped in any of the member functions.
+	 * This is only for tracking statistics and may not necessarily be accurate.
+	 * The count is never implicitly cleared by any member function. To clear
+	 * it call with bInClearToZero set to true.
+	 *
+	 * @return Number of samples that were dropped.
+	 * @see ClearNumDroppedSamples
+	 */
+	uint32 GetNumDroppedSamples(bool bInClearToZero)
+	{
+		uint32 n = NumDroppedSamples;
+		if (bInClearToZero)
+		{
+			NumDroppedSamples = 0;
+		}
+		return n;
+	}
+
 public:
 
 	//~ TMediaSampleSource interface (to be called only from consumer thread)
@@ -158,7 +177,9 @@ public:
 			check(LastPossibleIndex >= 0);
 
 			// Remove all samples indicated...
-			Samples.RemoveAt(FirstPossibleIndex, LastPossibleIndex - FirstPossibleIndex + 1);
+			const int32 NumToRemove = LastPossibleIndex - FirstPossibleIndex + 1;
+			Samples.RemoveAt(FirstPossibleIndex, NumToRemove);
+			NumDroppedSamples += (uint32)NumToRemove;
 			return true;
 		}
 		return false;
@@ -221,7 +242,9 @@ public:
 
 					// Found the best. Return it & delete all candidate samples up and including it from the queue
 					OutSample = Samples[BestIndex];
-					Samples.RemoveAt(FirstPossibleIndex, BestIndex - FirstPossibleIndex + 1);
+					int32 NumToRemove = BestIndex - FirstPossibleIndex + 1;
+					Samples.RemoveAt(FirstPossibleIndex, NumToRemove);
+					NumDroppedSamples += NumToRemove > 1 ? (uint32)(NumToRemove - 1) : 0;
 				}
 				else
 				{
@@ -254,7 +277,9 @@ public:
 					{
 						// Yes, so we return this one and remove all candidates and this one from the queue
 						OutSample = Sample;
-						Samples.RemoveAt(FirstPossibleIndex, Idx - FirstPossibleIndex + 1);
+						int32 NumToRemove = Idx - FirstPossibleIndex + 1;
+						Samples.RemoveAt(FirstPossibleIndex, NumToRemove);
+						NumDroppedSamples += NumToRemove > 1 ? (uint32)(NumToRemove - 1) : 0;
 						break;
 					}
 				}
@@ -266,6 +291,7 @@ public:
 		{
 			// Cleanup samples that are now considered outdated...
 			Samples.RemoveAt(0, NumOldSamplesAtBegin);
+			NumDroppedSamples += (uint32)NumOldSamplesAtBegin;
 		}
 
 		// Return true if we got a sample...
@@ -312,6 +338,7 @@ public:
 				{
 					// All earlier samples, including the current one are "too old"
 					Samples.RemoveAt(0, Idx + 1);
+					NumDroppedSamples += (uint32)(Idx + 1);
 					return Idx + 1;
 				}
 			}
@@ -325,6 +352,7 @@ public:
 		while(Samples.Num() && Samples[0]->GetTime().GetSequenceIndex() < InUntilIndex)
 		{
 			Samples.RemoveAt(0);
+			++NumDroppedSamples;
 		}
 	}
 
@@ -419,6 +447,7 @@ protected:
 	TArray<TSharedPtr<SampleType, ESPMode::ThreadSafe>> Samples;
 	int32 MaxSamplesInQueue;
 	uint32 FlushCount = 0;
+	uint32 NumDroppedSamples = 0;
 };
 
 
diff --git a/Engine/Source/Runtime/MediaUtils/Public/MediaSamples.h b/Engine/Source/Runtime/MediaUtils/Public/MediaSamples.h
--- a/Engine/Source/Runtime/MediaUtils/Public/MediaSamples.h
+++ b/Engine/Source/Runtime/MediaUtils/Public/MediaSamples.h
@@ -205,6 +205,32 @@ public:
 		return MetadataSampleQueue.Discard(TimeRange, bReverse);
 	}
 
+	virtual uint32 GetNumDroppedVideoSamples(bool bInClearToZero) override 
+	{ 
+		return VideoSampleQueue.GetNumDroppedSamples(bInClearToZero); 
+	}
+	
+	virtual uint32 GetNumDroppedAudioSamples(bool bInClearToZero) override 
+	{ 
+		return AudioSampleQueue.GetNumDroppedSamples(bInClearToZero); 
+	}
+	
+	virtual uint32 GetNumDroppedSubtitleSamples(bool bInClearToZero) override 
+	{ 
+		return SubtitleSampleQueue.GetNumDroppedSamples(bInClearToZero); 
+	}
+	
+	virtual uint32 GetNumDroppedCaptionSamples(bool bInClearToZero) override 
+	{ 
+		return CaptionSampleQueue.GetNumDroppedSamples(bInClearToZero); 
+	}
+	
+	virtual uint32 GetNumDroppedMetadataSamples(bool bInClearToZero) override 
+	{ 
+		return MetadataSampleQueue.GetNumDroppedSamples(bInClearToZero); 
+	}
+
+
 	MEDIAUTILS_API virtual uint32 PurgeOutdatedVideoSamples(const FMediaTimeStamp & ReferenceTime, bool bReversed, FTimespan MaxAge) override;
 	MEDIAUTILS_API virtual uint32 PurgeOutdatedSubtitleSamples(const FMediaTimeStamp & ReferenceTime, bool bReversed, FTimespan MaxAge) override;
 	MEDIAUTILS_API virtual uint32 PurgeOutdatedCaptionSamples(const FMediaTimeStamp& ReferenceTime, bool bReversed, FTimespan MaxAge) override;
