MediaUtils: corrected a comment

#rnx

diff --git a/Engine/Source/Runtime/MediaUtils/Public/MediaSampleQueue.h b/Engine/Source/Runtime/MediaUtils/Public/MediaSampleQueue.h
--- a/Engine/Source/Runtime/MediaUtils/Public/MediaSampleQueue.h
+++ b/Engine/Source/Runtime/MediaUtils/Public/MediaSampleQueue.h
@@ -70,7 +70,6 @@ public:
 	 * it call with bInClearToZero set to true.
 	 *
 	 * @return Number of samples that were dropped.
-	 * @see ClearNumDroppedSamples
 	 */
 	uint32 GetNumDroppedSamples(bool bInClearToZero)
 	{
