MediaUtils: calculating the wrapped modulo only when the duration is greater than zero

#rnx

diff --git a/Engine/Source/Runtime/MediaUtils/Private/MediaPlayerFacade.cpp b/Engine/Source/Runtime/MediaUtils/Private/MediaPlayerFacade.cpp
--- a/Engine/Source/Runtime/MediaUtils/Private/MediaPlayerFacade.cpp
+++ b/Engine/Source/Runtime/MediaUtils/Private/MediaPlayerFacade.cpp
@@ -128,7 +128,7 @@ namespace MediaPlayerFacade
 
 static FTimespan WrappedModulo(FTimespan Time, FTimespan Duration)
 {
-	return (Time >= FTimespan::Zero()) ? (Time % Duration) : (Duration + (Time % Duration));
+	return Duration > FTimespan::Zero() ? (Time >= FTimespan::Zero() ? (Time % Duration) : (Duration + (Time % Duration))) : Time;
 }
 
 static bool IsDurationValidAndFinite(FTimespan Duration)
