[MediaAssets] Support for video transcoding:  Adding a delegate post media texture render to allow rendering pipeline to be extended.

The primary use of this is to hook transcoding pipelines such as TMV (tiled-mip video).

#Synced-CL 49242762

diff --git a/Engine/Source/Runtime/MediaAssets/Private/Misc/MediaTextureResource.cpp b/Engine/Source/Runtime/MediaAssets/Private/Misc/MediaTextureResource.cpp
--- a/Engine/Source/Runtime/MediaAssets/Private/Misc/MediaTextureResource.cpp
+++ b/Engine/Source/Runtime/MediaAssets/Private/Misc/MediaTextureResource.cpp
@@ -667,6 +667,9 @@ void FMediaTextureResource::Render(const FRenderParams& Params)
 				}
 
 				SET_FLOAT_STAT(STAT_MediaUtils_TextureSampleTime, Sample->GetTime().Time.GetTotalMilliseconds());
+
+				// Notify extensions the media texture has been updated with a new sample.
+				FMediaDelegates::OnSampleRendered_RenderThread.Broadcast(this, Sample);
 			}
 		}
 	}
diff --git a/Engine/Source/Runtime/MediaUtils/Private/MediaDelegates.cpp b/Engine/Source/Runtime/MediaUtils/Private/MediaDelegates.cpp
--- a/Engine/Source/Runtime/MediaUtils/Private/MediaDelegates.cpp
+++ b/Engine/Source/Runtime/MediaUtils/Private/MediaDelegates.cpp
@@ -9,3 +9,5 @@ FMediaDelegates::FOnSampleDiscardedDelegate FMediaDelegates::OnSampleDiscarded_R
 FMediaDelegates::FOnPreSampleRenderDelegate FMediaDelegates::OnPreSampleRender_RenderThread;
 
 #endif //UE_MEDIAUTILS_DEVELOPMENT_DELEGATE
+
+FMediaDelegates::FOnSampleRenderedDelegate FMediaDelegates::OnSampleRendered_RenderThread;
\ No newline at end of file
diff --git a/Engine/Source/Runtime/MediaUtils/Public/MediaDelegates.h b/Engine/Source/Runtime/MediaUtils/Public/MediaDelegates.h
--- a/Engine/Source/Runtime/MediaUtils/Public/MediaDelegates.h
+++ b/Engine/Source/Runtime/MediaUtils/Public/MediaDelegates.h
@@ -9,6 +9,7 @@
 
 #define UE_MEDIAUTILS_DEVELOPMENT_DELEGATE (UE_BUILD_DEVELOPMENT || UE_BUILD_DEBUG)
 
+class FTextureResource;
 class IMediaTextureSample;
 class UMediaTexture;
 
@@ -19,7 +20,7 @@ public:
 
 	DECLARE_MULTICAST_DELEGATE_TwoParams(FOnSampleDiscardedDelegate, const UMediaTexture* /*Owner*/, TSharedMediaTextureSample /*FlushedSample*/);
 	DECLARE_MULTICAST_DELEGATE_ThreeParams(FOnPreSampleRenderDelegate, const UMediaTexture* /*Owner*/, bool /*UseSample*/, TSharedMediaTextureSample /*SampleToRender*/);
-
+	DECLARE_MULTICAST_DELEGATE_TwoParams(FOnSampleRenderedDelegate, const FTextureResource* /*MediaTextureResource*/, TSharedMediaTextureSample /*RenderedSample*/);
 
 #if UE_MEDIAUTILS_DEVELOPMENT_DELEGATE
 
@@ -38,4 +39,10 @@ public:
 	static MEDIAUTILS_API FOnPreSampleRenderDelegate OnPreSampleRender_RenderThread;
 
 #endif //UE_MEDIAUTILS_DEVELOPMENT_DELEGATE
+
+	/**
+	 * Callback when a sample has been rendered in the MediaTexture.
+	 * This can be used to extend the MediaTexture rendering pipeline. 
+	 */
+	static MEDIAUTILS_API FOnSampleRenderedDelegate OnSampleRendered_RenderThread;
 };
