Media: Convert new FOnSampleRenderedDelegate delegate to be thread-safe.

#Synced-CL 52172240

diff --git a/Engine/Source/Runtime/MediaUtils/Public/MediaDelegates.h b/Engine/Source/Runtime/MediaUtils/Public/MediaDelegates.h
--- a/Engine/Source/Runtime/MediaUtils/Public/MediaDelegates.h
+++ b/Engine/Source/Runtime/MediaUtils/Public/MediaDelegates.h
@@ -20,7 +20,7 @@ public:
 
 	DECLARE_MULTICAST_DELEGATE_TwoParams(FOnSampleDiscardedDelegate, const UMediaTexture* /*Owner*/, TSharedMediaTextureSample /*FlushedSample*/);
 	DECLARE_MULTICAST_DELEGATE_ThreeParams(FOnPreSampleRenderDelegate, const UMediaTexture* /*Owner*/, bool /*UseSample*/, TSharedMediaTextureSample /*SampleToRender*/);
-	DECLARE_MULTICAST_DELEGATE_TwoParams(FOnSampleRenderedDelegate, const FTextureResource* /*MediaTextureResource*/, TSharedMediaTextureSample /*RenderedSample*/);
+	DECLARE_TS_MULTICAST_DELEGATE_TwoParams(FOnSampleRenderedDelegate, const FTextureResource* /*MediaTextureResource*/, TSharedMediaTextureSample /*RenderedSample*/);
 
 #if UE_MEDIAUTILS_DEVELOPMENT_DELEGATE
 
